Create Blueprint Header View Plugin which adds a tab that contains a blueprint asset picker

diff --git a/Engine/Plugins/Editor/BlueprintHeaderView/BlueprintHeaderView.uplugin b/Engine/Plugins/Editor/BlueprintHeaderView/BlueprintHeaderView.uplugin
new file mode 100644
--- /dev/null
+++ b/Engine/Plugins/Editor/BlueprintHeaderView/BlueprintHeaderView.uplugin
@@ -0,0 +1,25 @@
+{
+	"FileVersion": 3,
+	"Version": 1,
+	"VersionName": "1.0",
+	"FriendlyName": "Blueprint Header View",
+	"Description": "A tool to help convert Blueprint Classes to Native C++.",
+	"Category": "Blueprints",
+	"CreatedBy": "Epic Games, Inc.",
+	"CreatedByURL": "http://epicgames.com",
+	"DocsURL": "",
+	"MarketplaceURL": "",
+	"SupportURL": "",
+	"CanContainContent": true,
+	"CanContainVerse": false,
+	"IsBetaVersion": false,
+	"IsExperimentalVersion": false,
+	"Installed": false,
+	"Modules": [
+		{
+			"Name": "BlueprintHeaderView",
+			"Type": "Editor",
+			"LoadingPhase": "Default"
+		}
+	]
+}
\ No newline at end of file
diff --git a/Engine/Plugins/Editor/BlueprintHeaderView/Source/BlueprintHeaderView/BlueprintHeaderView.Build.cs b/Engine/Plugins/Editor/BlueprintHeaderView/Source/BlueprintHeaderView/BlueprintHeaderView.Build.cs
new file mode 100644
--- /dev/null
+++ b/Engine/Plugins/Editor/BlueprintHeaderView/Source/BlueprintHeaderView/BlueprintHeaderView.Build.cs
@@ -0,0 +1,31 @@
+// Copyright Epic Games, Inc. All Rights Reserved.
+
+using UnrealBuildTool;
+
+public class BlueprintHeaderView : ModuleRules
+{
+	public BlueprintHeaderView(ReadOnlyTargetRules Target) : base(Target)
+	{
+		PCHUsage = ModuleRules.PCHUsageMode.UseExplicitOrSharedPCHs;
+		
+		PublicDependencyModuleNames.AddRange(
+			new string[]
+			{
+				"Core",
+			}
+			);
+			
+		
+		PrivateDependencyModuleNames.AddRange(
+			new string[]
+			{
+				"CoreUObject",
+				"Engine",
+				"Slate",
+				"SlateCore",
+				"EditorStyle",
+			}
+			);
+		
+	}
+}
diff --git a/Engine/Plugins/Editor/BlueprintHeaderView/Source/BlueprintHeaderView/Private/BlueprintHeaderView.cpp b/Engine/Plugins/Editor/BlueprintHeaderView/Source/BlueprintHeaderView/Private/BlueprintHeaderView.cpp
new file mode 100644
--- /dev/null
+++ b/Engine/Plugins/Editor/BlueprintHeaderView/Source/BlueprintHeaderView/Private/BlueprintHeaderView.cpp
@@ -0,0 +1,45 @@
+// Copyright Epic Games, Inc. All Rights Reserved.
+
+#include "BlueprintHeaderView.h"
+#include "Widgets/SWidget.h"
+#include "Modules/ModuleManager.h"
+#include "Widgets/DeclarativeSyntaxSupport.h"
+#include "Framework/Application/SlateApplication.h"
+#include "Textures/SlateIcon.h"
+#include "Framework/Docking/TabManager.h"
+#include "EditorStyleSet.h"
+#include "Editor/WorkspaceMenuStructure/Public/WorkspaceMenuStructure.h"
+#include "Editor/WorkspaceMenuStructure/Public/WorkspaceMenuStructureModule.h"
+#include "Widgets/Docking/SDockTab.h"
+#include "SBlueprintHeaderView.h"
+
+namespace BlueprintHeaderViewModule
+{
+	static const FName HeaderViewTabName = "BlueprintHeaderViewApp";
+
+	TSharedRef<SDockTab> CreateHeaderViewTab(const FSpawnTabArgs& Args)
+	{
+		return SNew(SDockTab)
+			.TabRole(ETabRole::NomadTab)
+			[
+				SNew(SBlueprintHeaderView)
+			];
+	}
+}
+
+void FBlueprintHeaderViewModule::StartupModule()
+{
+	FGlobalTabmanager::Get()->RegisterNomadTabSpawner(BlueprintHeaderViewModule::HeaderViewTabName, FOnSpawnTab::CreateStatic(&BlueprintHeaderViewModule::CreateHeaderViewTab))
+		.SetDisplayName(NSLOCTEXT("BlueprintHeaderViewApp", "TabTitle", "Blueprint Header View"))
+		.SetTooltipText(NSLOCTEXT("BlueprintHeaderViewApp", "TooltipText", "Displays a Blueprint Class in C++ Header format."))
+		.SetGroup(WorkspaceMenu::GetMenuStructure().GetToolsCategory())
+		.SetIcon(FSlateIcon(FEditorStyle::GetStyleSetName(), "ClassIcon.Class"));
+}
+
+void FBlueprintHeaderViewModule::ShutdownModule()
+{
+	// This function may be called during shutdown to clean up your module.  For modules that support dynamic reloading,
+	// we call this function before unloading the module.
+}
+	
+IMPLEMENT_MODULE(FBlueprintHeaderViewModule, BlueprintHeaderView)
\ No newline at end of file
diff --git a/Engine/Plugins/Editor/BlueprintHeaderView/Source/BlueprintHeaderView/Private/SBlueprintHeaderView.cpp b/Engine/Plugins/Editor/BlueprintHeaderView/Source/BlueprintHeaderView/Private/SBlueprintHeaderView.cpp
new file mode 100644
--- /dev/null
+++ b/Engine/Plugins/Editor/BlueprintHeaderView/Source/BlueprintHeaderView/Private/SBlueprintHeaderView.cpp
@@ -0,0 +1,100 @@
+// Copyright Epic Games, Inc. All Rights Reserved.
+
+#include "SBlueprintHeaderView.h"
+#include "Widgets/Layout/SSpacer.h"
+#include "Widgets/Layout/SBox.h"
+#include "Widgets/Input/SComboButton.h"
+#include "ContentBrowserModule.h"
+#include "IContentBrowserSingleton.h"
+#include "Engine/Blueprint.h"
+
+#define LOCTEXT_NAMESPACE "SBlueprintHeaderView"
+
+void SBlueprintHeaderView::Construct(const FArguments& InArgs)
+{
+	const float PaddingAmount = 8.0f;
+	SelectedBlueprint = nullptr;
+
+	ChildSlot
+	[
+		SNew(SVerticalBox)
+		+SVerticalBox::Slot()
+		.AutoHeight()
+		.Padding(FMargin(PaddingAmount))
+		[
+			SNew(SHorizontalBox)
+			+SHorizontalBox::Slot()
+			.AutoWidth()
+			.VAlign(VAlign_Center)
+			.HAlign(HAlign_Left)
+			[
+				SNew(STextBlock)
+				.Text(LOCTEXT("ClassPickerLabel", "Displaying Blueprint:"))
+			]
+			+SHorizontalBox::Slot()
+			.AutoWidth()
+			.HAlign(HAlign_Left)
+			[
+				SNew(SSpacer)
+				.Size(FVector2D(PaddingAmount))
+			]
+			+SHorizontalBox::Slot()
+			.AutoWidth()
+			.HAlign(HAlign_Left)
+			[
+				SNew(SBox)
+				.WidthOverride(400.0f)
+				[
+					SAssignNew(ClassPickerComboButton, SComboButton)
+					.OnGetMenuContent(this, &SBlueprintHeaderView::GetClassPickerMenuContent)
+					.ButtonContent()
+					[
+						SNew(STextBlock)
+						.Text(this, &SBlueprintHeaderView::GetClassPickerText)
+					]
+				]
+			]
+		]
+		+SVerticalBox::Slot()
+		[
+			SNullWidget::NullWidget // TODO: UE-138705 replace this with a list view and start adding items to it
+		]
+	];
+}
+
+FText SBlueprintHeaderView::GetClassPickerText() const
+{
+	if (const UBlueprint* Blueprint = SelectedBlueprint.Get())
+	{
+		return FText::FromName(Blueprint->GetFName());
+	}
+
+	return LOCTEXT("ClassPickerPickClass", "Select Blueprint Class");
+}
+
+TSharedRef<SWidget> SBlueprintHeaderView::GetClassPickerMenuContent()
+{
+	FContentBrowserModule& ContentBrowserModule = FModuleManager::Get().LoadModuleChecked<FContentBrowserModule>("ContentBrowser");
+
+	FAssetPickerConfig AssetPickerConfig;
+	AssetPickerConfig.SelectionMode = ESelectionMode::Single;
+	AssetPickerConfig.OnAssetSelected = FOnAssetSelected::CreateSP(this, &SBlueprintHeaderView::OnAssetSelected);
+	AssetPickerConfig.Filter.ClassNames.Add(UBlueprint::StaticClass()->GetFName());
+	AssetPickerConfig.InitialAssetViewType = EAssetViewType::List;
+	TSharedRef<SWidget> AssetPickerWidget = ContentBrowserModule.Get().CreateAssetPicker(AssetPickerConfig);
+
+	return SNew(SBox)
+		.HeightOverride(500.f)
+		[
+			AssetPickerWidget
+		];
+}
+
+void SBlueprintHeaderView::OnAssetSelected(const FAssetData& SelectedAsset)
+{
+	ClassPickerComboButton->SetIsOpen(false);
+
+	SelectedBlueprint = Cast<UBlueprint>(SelectedAsset.GetAsset());
+}
+
+#undef LOCTEXT_NAMESPACE
diff --git a/Engine/Plugins/Editor/BlueprintHeaderView/Source/BlueprintHeaderView/Private/SBlueprintHeaderView.h b/Engine/Plugins/Editor/BlueprintHeaderView/Source/BlueprintHeaderView/Private/SBlueprintHeaderView.h
new file mode 100644
--- /dev/null
+++ b/Engine/Plugins/Editor/BlueprintHeaderView/Source/BlueprintHeaderView/Private/SBlueprintHeaderView.h
@@ -0,0 +1,40 @@
+// Copyright Epic Games, Inc. All Rights Reserved.
+
+#pragma once
+
+#include "CoreMinimal.h"
+#include "Widgets/DeclarativeSyntaxSupport.h"
+#include "Widgets/SWidget.h"
+#include "Widgets/SCompoundWidget.h"
+
+class SComboButton;
+class UBlueprint;
+
+class SBlueprintHeaderView : public SCompoundWidget
+{
+public:
+	SLATE_BEGIN_ARGS(SBlueprintHeaderView)
+	{
+	}
+
+	SLATE_END_ARGS()
+
+	void Construct(const FArguments& InArgs);
+
+private:
+	/** Gets the text for the class picker combo button */
+	FText GetClassPickerText() const;
+
+	/** Constructs a Blueprint Class picker menu widget */
+	TSharedRef<SWidget> GetClassPickerMenuContent();
+
+	/** Callback for class picker menu selecting a blueprint asset */
+	void OnAssetSelected(const FAssetData& SelectedAsset);
+
+private:
+	/** The blueprint currently being displayed by the header view */
+	TWeakObjectPtr<UBlueprint> SelectedBlueprint;
+
+	/** Reference to the Class Picker combo button widget */
+	TSharedPtr<SComboButton> ClassPickerComboButton;
+};
\ No newline at end of file
diff --git a/Engine/Plugins/Editor/BlueprintHeaderView/Source/BlueprintHeaderView/Public/BlueprintHeaderView.h b/Engine/Plugins/Editor/BlueprintHeaderView/Source/BlueprintHeaderView/Public/BlueprintHeaderView.h
new file mode 100644
--- /dev/null
+++ b/Engine/Plugins/Editor/BlueprintHeaderView/Source/BlueprintHeaderView/Public/BlueprintHeaderView.h
@@ -0,0 +1,15 @@
+// Copyright Epic Games, Inc. All Rights Reserved.
+
+#pragma once
+
+#include "CoreMinimal.h"
+#include "Modules/ModuleManager.h"
+
+class FBlueprintHeaderViewModule : public IModuleInterface
+{
+public:
+
+	/** IModuleInterface implementation */
+	virtual void StartupModule() override;
+	virtual void ShutdownModule() override;
+};
