Add class declaration to the Blueprint Header View

diff --git a/Engine/Plugins/Editor/BlueprintHeaderView/Source/BlueprintHeaderView/BlueprintHeaderView.Build.cs b/Engine/Plugins/Editor/BlueprintHeaderView/Source/BlueprintHeaderView/BlueprintHeaderView.Build.cs
--- a/Engine/Plugins/Editor/BlueprintHeaderView/Source/BlueprintHeaderView/BlueprintHeaderView.Build.cs
+++ b/Engine/Plugins/Editor/BlueprintHeaderView/Source/BlueprintHeaderView/BlueprintHeaderView.Build.cs
@@ -21,6 +21,7 @@ public class BlueprintHeaderView : ModuleRules
 			{
 				"CoreUObject",
 				"Engine",
+				"InputCore",
 				"Slate",
 				"SlateCore",
 				"EditorStyle",
diff --git a/Engine/Plugins/Editor/BlueprintHeaderView/Source/BlueprintHeaderView/Private/HeaderViewClassListItem.cpp b/Engine/Plugins/Editor/BlueprintHeaderView/Source/BlueprintHeaderView/Private/HeaderViewClassListItem.cpp
new file mode 100644
--- /dev/null
+++ b/Engine/Plugins/Editor/BlueprintHeaderView/Source/BlueprintHeaderView/Private/HeaderViewClassListItem.cpp
@@ -0,0 +1,127 @@
+// Copyright Epic Games, Inc. All Rights Reserved.
+
+#include "HeaderViewClassListItem.h"
+#include "Engine/Blueprint.h"
+#include "String/LineEndings.h"
+
+#define LOCTEXT_NAMESPACE "HeaderViewClassListItem"
+
+FHeaderViewListItemPtr FHeaderViewClassListItem::Create(TWeakObjectPtr<UBlueprint> InBlueprint)
+{
+	return MakeShareable(new FHeaderViewClassListItem(InBlueprint));
+}
+
+FString FHeaderViewClassListItem::GetConditionalUClassSpecifiers(const UBlueprint* Blueprint) const
+{
+	// Avoid lots of reallocations
+	TStringBuilder<256> AdditionalSpecifiers;
+
+	if (Blueprint->bGenerateConstClass)
+	{
+		AdditionalSpecifiers.Append(TEXT(", Const"));
+	}
+
+	if (Blueprint->bGenerateAbstractClass)
+	{
+		AdditionalSpecifiers.Append(TEXT(", Abstract"));
+	}
+
+	if (!Blueprint->BlueprintCategory.IsEmpty())
+	{
+		AdditionalSpecifiers.Append(FString::Printf(TEXT(", Category=\"%s\""), *Blueprint->BlueprintCategory));
+	}
+
+	if (!Blueprint->HideCategories.IsEmpty())
+	{
+		AdditionalSpecifiers.Append(TEXT(", HideCategories=("));
+
+		for (int32 HideCategoryIdx = 0; HideCategoryIdx < Blueprint->HideCategories.Num(); ++HideCategoryIdx)
+		{
+			if (HideCategoryIdx == 0)
+			{
+				AdditionalSpecifiers.Append(FString::Printf(TEXT("\"%s\""), *Blueprint->HideCategories[HideCategoryIdx]));
+			}
+			else
+			{
+				AdditionalSpecifiers.Append(FString::Printf(TEXT(", \"%s\""), *Blueprint->HideCategories[HideCategoryIdx]));
+			}
+		}
+
+		AdditionalSpecifiers.Append(TEXT(")"));
+	}
+
+	if (!Blueprint->BlueprintDisplayName.IsEmpty() || !Blueprint->BlueprintNamespace.IsEmpty())
+	{
+		AdditionalSpecifiers.Append(TEXT(", meta=("));
+
+		int32 SpecifierCount = 0;
+
+		if (!Blueprint->BlueprintDisplayName.IsEmpty())
+		{
+			AdditionalSpecifiers.Append(FString::Printf(TEXT("DisplayName=\"%s\""), *Blueprint->BlueprintDisplayName));
+			++SpecifierCount;
+		}
+
+		if (!Blueprint->BlueprintNamespace.IsEmpty())
+		{
+			if (SpecifierCount > 0)
+			{
+				AdditionalSpecifiers.Append(FString::Printf(TEXT(", Namespace=\"%s\""), *Blueprint->BlueprintNamespace));
+			}
+			else
+			{
+				AdditionalSpecifiers.Append(FString::Printf(TEXT("Namespace=\"%s\""), *Blueprint->BlueprintNamespace));
+			}
+		}
+
+		AdditionalSpecifiers.Append(TEXT(")"));
+	}
+
+	return AdditionalSpecifiers.ToString();
+}
+
+FHeaderViewClassListItem::FHeaderViewClassListItem(TWeakObjectPtr<UBlueprint> InBlueprint)
+{
+	if (UBlueprint* Blueprint = InBlueprint.Get())
+	{
+		// Avoid lots of reallocations
+		RawItemString.Reserve(512);
+		RichTextString.Reserve(512);
+
+		// Format comment
+		{
+			FString Comment = Blueprint->BlueprintDescription.IsEmpty() ? TEXT("Please add a class description") : Blueprint->BlueprintDescription;
+			FormatCommentString(Comment, RawItemString, RichTextString);
+		}
+
+		// Add the UCLASS specifiers
+		// i.e. UCLASS(Blueprintable, BlueprintType, Category=BlueprintCategory)
+		{
+			FString AdditionalSpecifiers = GetConditionalUClassSpecifiers(Blueprint);
+
+			// Always add Blueprintable and BlueprintType
+			RawItemString += FString::Printf(TEXT("\nUCLASS(Blueprintable, BlueprintType%s)"), *AdditionalSpecifiers);
+			RichTextString += FString::Printf(TEXT("\n<macro>UCLASS</>(Blueprintable, BlueprintType%s)"), *AdditionalSpecifiers);
+		}
+
+		// Add the class declaration line
+		// i.e. class ClassName : public ParentClass
+		{
+			const FString BlueprintName = Blueprint->SkeletonGeneratedClass->GetPrefixCPP() + Blueprint->GetName();
+
+			const UObject* ParentBlueprint = Blueprint->ParentClass->ClassGeneratedBy;
+			FString ParentClassName = FString::Printf(TEXT("%s%s"), 
+				ParentBlueprint ? TEXT("") : Blueprint->ParentClass->GetPrefixCPP(),
+				ParentBlueprint ? *ParentBlueprint->GetName() : *Blueprint->ParentClass->GetAuthoredName());
+
+			RawItemString += FString::Printf(TEXT("\nclass %s : public %s\n{\n\tGENERATED_BODY()"), *BlueprintName, *ParentClassName);
+			RichTextString += FString::Printf(TEXT("\n<keyword>class</> <typename>%s</> : <keyword>public</> <typename>%s</>\n{\n\t<macro>GENERATED_BODY</>()"), *BlueprintName, *ParentClassName);
+		}
+
+		// normalize to platform newlines
+		UE::String::ToHostLineEndingsInline(RawItemString);
+		UE::String::ToHostLineEndingsInline(RichTextString);
+	}
+}
+
+#undef LOCTEXT_NAMESPACE
diff --git a/Engine/Plugins/Editor/BlueprintHeaderView/Source/BlueprintHeaderView/Private/HeaderViewClassListItem.h b/Engine/Plugins/Editor/BlueprintHeaderView/Source/BlueprintHeaderView/Private/HeaderViewClassListItem.h
new file mode 100644
--- /dev/null
+++ b/Engine/Plugins/Editor/BlueprintHeaderView/Source/BlueprintHeaderView/Private/HeaderViewClassListItem.h
@@ -0,0 +1,20 @@
+// Copyright Epic Games, Inc. All Rights Reserved.
+
+#pragma once
+
+#include "CoreMinimal.h"
+#include "SBlueprintHeaderView.h"
+
+/** A header view list item that displays the class declaration */
+struct FHeaderViewClassListItem : public FHeaderViewListItem
+{
+public:
+	/** Creates a list item for the Header view representing a class declaration for the given blueprint */
+	static FHeaderViewListItemPtr Create(TWeakObjectPtr<UBlueprint> InBlueprint);
+
+protected:
+	FString GetConditionalUClassSpecifiers(const UBlueprint* Blueprint) const;
+
+	FHeaderViewClassListItem(TWeakObjectPtr<UBlueprint> InBlueprint);
+
+};
\ No newline at end of file
diff --git a/Engine/Plugins/Editor/BlueprintHeaderView/Source/BlueprintHeaderView/Private/SBlueprintHeaderView.cpp b/Engine/Plugins/Editor/BlueprintHeaderView/Source/BlueprintHeaderView/Private/SBlueprintHeaderView.cpp
--- a/Engine/Plugins/Editor/BlueprintHeaderView/Source/BlueprintHeaderView/Private/SBlueprintHeaderView.cpp
+++ b/Engine/Plugins/Editor/BlueprintHeaderView/Source/BlueprintHeaderView/Private/SBlueprintHeaderView.cpp
@@ -3,13 +3,141 @@
 #include "SBlueprintHeaderView.h"
 #include "Widgets/Layout/SSpacer.h"
 #include "Widgets/Layout/SBox.h"
+//#include "Widgets/Layout/SScrollBox.h"
 #include "Widgets/Input/SComboButton.h"
+#include "Widgets/Text/SRichTextBlock.h"
 #include "ContentBrowserModule.h"
 #include "IContentBrowserSingleton.h"
 #include "Engine/Blueprint.h"
+#include "Framework/Text/SlateTextRun.h"
+#include "HeaderViewClassListItem.h"
+#include "EditorStyleSet.h"
+#include "String/LineEndings.h"
 
 #define LOCTEXT_NAMESPACE "SBlueprintHeaderView"
 
+// FHeaderViewSyntaxDecorator /////////////////////////////////////////////////
+
+namespace
+{
+	class FHeaderViewSyntaxDecorator : public ITextDecorator
+	{
+	public:
+
+		static TSharedRef<FHeaderViewSyntaxDecorator> Create(FString InName, const FSlateColor& InColor)
+		{
+			return MakeShareable(new FHeaderViewSyntaxDecorator(MoveTemp(InName), InColor));
+		}
+
+		bool Supports(const FTextRunParseResults& RunInfo, const FString& Text) const override
+		{
+			return RunInfo.Name == DecoratorName;
+		}
+
+		TSharedRef<ISlateRun> Create(const TSharedRef<class FTextLayout>& TextLayout, const FTextRunParseResults& RunParseResult, const FString& OriginalText, const TSharedRef<FString>& ModelText, const ISlateStyle* Style) override
+		{
+			FRunInfo RunInfo(RunParseResult.Name);
+			for (const TPair<FString, FTextRange>& Pair : RunParseResult.MetaData)
+			{
+				RunInfo.MetaData.Add(Pair.Key, OriginalText.Mid(Pair.Value.BeginIndex, Pair.Value.EndIndex - Pair.Value.BeginIndex));
+			}
+
+			ModelText->Append(OriginalText.Mid(RunParseResult.ContentRange.BeginIndex, RunParseResult.ContentRange.Len()));
+
+			return FSlateTextRun::Create(RunInfo, ModelText, TextStyle);
+		}
+
+	private:
+
+		FHeaderViewSyntaxDecorator(FString&& InName, const FSlateColor& InColor)
+			: DecoratorName(MoveTemp(InName))
+			, TextStyle(FEditorStyle::Get().GetWidgetStyle<FTextBlockStyle>("Log.Normal"))
+		{
+			TextStyle.SetColorAndOpacity(InColor);
+		}
+
+	private:
+
+		/** Name of this decorator */
+		FString DecoratorName;
+
+		/** Style of this decorator */
+		FTextBlockStyle TextStyle;
+
+	};
+
+}
+
+// FHeaderViewListItem ////////////////////////////////////////////////////////
+
+TSharedRef<SWidget> FHeaderViewListItem::GenerateWidgetForItem()
+{
+	// TODO: colors not final, move this struct into config
+	static const struct
+	{
+		FSlateColor Comment = FLinearColor(0.3f, 0.7f, 0.1f, 1.0f);
+		FSlateColor Macro = FLinearColor(0.6f, 0.2f, 0.8f, 1.0f);
+		FSlateColor Typename = FLinearColor::White;
+		FSlateColor Identifier = FLinearColor::White;
+		FSlateColor Keyword = FLinearColor(0.0f, 0.4f, 0.8f, 1.0f);
+	} SyntaxColors;
+
+	return SNew(SBox)
+		.HAlign(HAlign_Fill)
+		.Padding(FMargin(4.0f))
+		[
+			SNew(SRichTextBlock)
+			.Text(FText::FromString(RichTextString))
+			.TextStyle(FEditorStyle::Get(), "Log.Normal")
+			+SRichTextBlock::Decorator(FHeaderViewSyntaxDecorator::Create(TEXT("comment"), SyntaxColors.Comment))
+			+SRichTextBlock::Decorator(FHeaderViewSyntaxDecorator::Create(TEXT("macro"), SyntaxColors.Macro))
+			+SRichTextBlock::Decorator(FHeaderViewSyntaxDecorator::Create(TEXT("typename"), SyntaxColors.Typename))
+			+SRichTextBlock::Decorator(FHeaderViewSyntaxDecorator::Create(TEXT("identifier"), SyntaxColors.Identifier))
+			+SRichTextBlock::Decorator(FHeaderViewSyntaxDecorator::Create(TEXT("keyword"), SyntaxColors.Keyword))
+		];
+}
+
+TSharedPtr<FHeaderViewListItem> FHeaderViewListItem::Create(FString InRawString, FString InRichText)
+{
+	return MakeShareable(new FHeaderViewListItem(MoveTemp(InRawString), MoveTemp(InRichText)));
+}
+
+FHeaderViewListItem::FHeaderViewListItem(FString&& InRawString, FString&& InRichText)
+	: RichTextString(MoveTemp(InRichText))
+	, RawItemString(MoveTemp(InRawString))
+{
+}
+
+void FHeaderViewListItem::FormatCommentString(FString InComment, FString& OutRawString, FString& OutRichString)
+{
+	// normalize newlines to \n
+	UE::String::FromHostLineEndingsInline(InComment);
+
+	if (InComment.Contains("\n"))
+	{
+		/**
+		 * Format into a multiline C++ comment, like this one
+		 */
+		InComment = TEXT("/**\n") + InComment;
+		InComment.ReplaceInline(TEXT("\n"), TEXT("\n * "));
+		InComment.Append(TEXT("\n */"));
+	}
+	else
+	{
+		/** Format into a single line C++ comment, like this one */
+		InComment = FString::Printf(TEXT("/** %s */"), *InComment);
+	}
+
+	// add the comment to the raw string representation
+	OutRawString = InComment;
+
+	// mark each line of the comment as the beginning and end of a comment style for the rich text representation
+	InComment.ReplaceInline(TEXT("\n"), TEXT("</>\n<comment>"));
+	OutRichString = FString::Printf(TEXT("<comment>%s</>"), *InComment);
+}
+
+// SBlueprintHeaderView ///////////////////////////////////////////////////////
+
 void SBlueprintHeaderView::Construct(const FArguments& InArgs)
 {
 	const float PaddingAmount = 8.0f;
@@ -56,8 +184,17 @@ void SBlueprintHeaderView::Construct(const FArguments& InArgs)
 			]
 		]
 		+SVerticalBox::Slot()
+		.Padding(FMargin(PaddingAmount))
 		[
-			SNullWidget::NullWidget // TODO: UE-138705 replace this with a list view and start adding items to it
+			//TODO: Add a scroll bar when this overflows from the available space
+// 			SNew(SScrollBox)
+// 			.Orientation(Orient_Horizontal)
+// 			+SScrollBox::Slot()
+// 			[
+				SAssignNew(ListView, SListView<FHeaderViewListItemPtr>)
+				.ListItemsSource(&ListItems)
+				.OnGenerateRow(this, &SBlueprintHeaderView::GenerateRowForItem)
+//			]
 		]
 	];
 }
@@ -95,6 +232,31 @@ void SBlueprintHeaderView::OnAssetSelected(const FAssetData& SelectedAsset)
 	ClassPickerComboButton->SetIsOpen(false);
 
 	SelectedBlueprint = Cast<UBlueprint>(SelectedAsset.GetAsset());
+
+	RepopulateListView();
+}
+
+TSharedRef<ITableRow> SBlueprintHeaderView::GenerateRowForItem(FHeaderViewListItemPtr Item, const TSharedRef<STableViewBase>& OwnerTable) const
+{
+	return SNew(STableRow<FHeaderViewListItemPtr>, OwnerTable)
+		.Content()
+		[
+			Item->GenerateWidgetForItem()
+		];
+}
+
+void SBlueprintHeaderView::RepopulateListView()
+{
+	ListItems.Empty();
+
+	// Add the class declaration
+	ListItems.Add(FHeaderViewClassListItem::Create(SelectedBlueprint));
+
+
+	// Add the closing brace of the class
+	ListItems.Add(FHeaderViewListItem::Create(TEXT("};"), TEXT("};")));
+	
+	ListView->RequestListRefresh();
 }
 
 #undef LOCTEXT_NAMESPACE
diff --git a/Engine/Plugins/Editor/BlueprintHeaderView/Source/BlueprintHeaderView/Private/SBlueprintHeaderView.h b/Engine/Plugins/Editor/BlueprintHeaderView/Source/BlueprintHeaderView/Private/SBlueprintHeaderView.h
--- a/Engine/Plugins/Editor/BlueprintHeaderView/Source/BlueprintHeaderView/Private/SBlueprintHeaderView.h
+++ b/Engine/Plugins/Editor/BlueprintHeaderView/Source/BlueprintHeaderView/Private/SBlueprintHeaderView.h
@@ -6,9 +6,46 @@
 #include "Widgets/DeclarativeSyntaxSupport.h"
 #include "Widgets/SWidget.h"
 #include "Widgets/SCompoundWidget.h"
+#include "Widgets/Views/SListView.h"
 
 class SComboButton;
 class UBlueprint;
+class ITableRow;
+class STableViewBase;
+
+/** A base class for List Items in the Header View */
+struct FHeaderViewListItem : public TSharedFromThis<FHeaderViewListItem>
+{
+	/** Creates the widget for this list item */
+	TSharedRef<SWidget> GenerateWidgetForItem();
+
+	/** Creates a basic list item containing some text */
+	static TSharedPtr<FHeaderViewListItem> Create(FString InRawString, FString InRichText);
+
+protected:
+	/** Empty base constructor hidden from public */
+	FHeaderViewListItem() {};
+
+	FHeaderViewListItem(FString&& InRawString, FString&& InRichText);
+
+	/** 
+	 * Formats a string into a C++ comment
+	 * @param InComment The string to format as a C++ comment
+	 * @param OutRawString The string formatted as a C++ comment
+	 * @param OutRichString The string formatted as a C++ comment with rich text decorators for syntax highlighting
+	 */
+	static void FormatCommentString(FString InComment, FString& OutRawString, FString& OutRichString);
+
+protected:
+	/** A rich text representation of the item, including syntax highlighting and errors */
+	FString RichTextString;
+
+	/** A raw string representation of the item, used for copying the item */
+	FString RawItemString;
+
+};
+
+using FHeaderViewListItemPtr = TSharedPtr<FHeaderViewListItem>;
 
 class SBlueprintHeaderView : public SCompoundWidget
 {
@@ -31,10 +68,21 @@ private:
 	/** Callback for class picker menu selecting a blueprint asset */
 	void OnAssetSelected(const FAssetData& SelectedAsset);
 
+	/** Generates a row for a given List Item */
+	TSharedRef<ITableRow> GenerateRowForItem(FHeaderViewListItemPtr Item, const TSharedRef<STableViewBase>& OwnerTable) const;
+	
+	/** Clears the list and repopulates it with info for the selected Blueprint */
+	void RepopulateListView();
 private:
 	/** The blueprint currently being displayed by the header view */
 	TWeakObjectPtr<UBlueprint> SelectedBlueprint;
 
 	/** Reference to the Class Picker combo button widget */
 	TSharedPtr<SComboButton> ClassPickerComboButton;
+
+	/** Reference to the ListView Widget*/
+	TSharedPtr<SListView<FHeaderViewListItemPtr>> ListView;
+
+	/** List Items source */
+	TArray<FHeaderViewListItemPtr> ListItems;
 };
\ No newline at end of file
