BlueprintHeaderView: Add Function declarations to the view

diff --git a/Engine/Plugins/Editor/BlueprintHeaderView/Source/BlueprintHeaderView/BlueprintHeaderView.Build.cs b/Engine/Plugins/Editor/BlueprintHeaderView/Source/BlueprintHeaderView/BlueprintHeaderView.Build.cs
--- a/Engine/Plugins/Editor/BlueprintHeaderView/Source/BlueprintHeaderView/BlueprintHeaderView.Build.cs
+++ b/Engine/Plugins/Editor/BlueprintHeaderView/Source/BlueprintHeaderView/BlueprintHeaderView.Build.cs
@@ -25,6 +25,8 @@ public class BlueprintHeaderView : ModuleRules
 				"Slate",
 				"SlateCore",
 				"EditorStyle",
+				"UnrealEd",
+				"BlueprintGraph",
 			}
 			);
 		
diff --git a/Engine/Plugins/Editor/BlueprintHeaderView/Source/BlueprintHeaderView/Private/HeaderViewClassListItem.cpp b/Engine/Plugins/Editor/BlueprintHeaderView/Source/BlueprintHeaderView/Private/HeaderViewClassListItem.cpp
--- a/Engine/Plugins/Editor/BlueprintHeaderView/Source/BlueprintHeaderView/Private/HeaderViewClassListItem.cpp
+++ b/Engine/Plugins/Editor/BlueprintHeaderView/Source/BlueprintHeaderView/Private/HeaderViewClassListItem.cpp
@@ -4,8 +4,6 @@
 #include "Engine/Blueprint.h"
 #include "String/LineEndings.h"
 
-#define LOCTEXT_NAMESPACE "HeaderViewClassListItem"
-
 FHeaderViewListItemPtr FHeaderViewClassListItem::Create(TWeakObjectPtr<UBlueprint> InBlueprint)
 {
 	return MakeShareable(new FHeaderViewClassListItem(InBlueprint));
@@ -13,7 +11,6 @@ FHeaderViewListItemPtr FHeaderViewClassListItem::Create(TWeakObjectPtr<UBlueprin
 
 FString FHeaderViewClassListItem::GetConditionalUClassSpecifiers(const UBlueprint* Blueprint) const
 {
-	// Avoid lots of reallocations
 	TStringBuilder<256> AdditionalSpecifiers;
 
 	if (Blueprint->bGenerateConstClass)
@@ -101,7 +98,7 @@ FHeaderViewClassListItem::FHeaderViewClassListItem(TWeakObjectPtr<UBlueprint> In
 
 			// Always add Blueprintable and BlueprintType
 			RawItemString += FString::Printf(TEXT("\nUCLASS(Blueprintable, BlueprintType%s)"), *AdditionalSpecifiers);
-			RichTextString += FString::Printf(TEXT("\n<macro>UCLASS</>(Blueprintable, BlueprintType%s)"), *AdditionalSpecifiers);
+			RichTextString += FString::Printf(TEXT("\n<%s>UCLASS</>(Blueprintable, BlueprintType%s)"), *HeaderViewSyntaxDecorators::MacroDecorator, *AdditionalSpecifiers);
 		}
 
 		// Add the class declaration line
@@ -115,7 +112,15 @@ FHeaderViewClassListItem::FHeaderViewClassListItem(TWeakObjectPtr<UBlueprint> In
 				ParentBlueprint ? *ParentBlueprint->GetName() : *Blueprint->ParentClass->GetAuthoredName());
 
 			RawItemString += FString::Printf(TEXT("\nclass %s : public %s\n{\n\tGENERATED_BODY()"), *BlueprintName, *ParentClassName);
-			RichTextString += FString::Printf(TEXT("\n<keyword>class</> <typename>%s</> : <keyword>public</> <typename>%s</>\n{\n\t<macro>GENERATED_BODY</>()"), *BlueprintName, *ParentClassName);
+			RichTextString += FString::Printf(TEXT("\n<%s>class</> <%s>%s</> : <%s>public</> <%s>%s</>\n{\n\t<%s>GENERATED_BODY</>()"), 
+				*HeaderViewSyntaxDecorators::KeywordDecorator,
+				*HeaderViewSyntaxDecorators::TypenameDecorator, 
+				*BlueprintName,
+				*HeaderViewSyntaxDecorators::KeywordDecorator,
+				*HeaderViewSyntaxDecorators::TypenameDecorator,
+				*ParentClassName,
+				*HeaderViewSyntaxDecorators::MacroDecorator
+			);
 		}
 
 		// normalize to platform newlines
@@ -124,4 +129,3 @@ FHeaderViewClassListItem::FHeaderViewClassListItem(TWeakObjectPtr<UBlueprint> In
 	}
 }
 
-#undef LOCTEXT_NAMESPACE
diff --git a/Engine/Plugins/Editor/BlueprintHeaderView/Source/BlueprintHeaderView/Private/HeaderViewFunctionListItem.cpp b/Engine/Plugins/Editor/BlueprintHeaderView/Source/BlueprintHeaderView/Private/HeaderViewFunctionListItem.cpp
new file mode 100644
--- /dev/null
+++ b/Engine/Plugins/Editor/BlueprintHeaderView/Source/BlueprintHeaderView/Private/HeaderViewFunctionListItem.cpp
@@ -0,0 +1,236 @@
+// Copyright Epic Games, Inc. All Rights Reserved.
+
+#include "HeaderViewFunctionListItem.h"
+#include "Engine/Blueprint.h"
+#include "String/LineEndings.h"
+#include "Kismet2/BlueprintEditorUtils.h"
+#include "K2Node_FunctionEntry.h"
+#include "Misc/EngineVersion.h"
+
+FHeaderViewListItemPtr FHeaderViewFunctionListItem::Create(const UK2Node_FunctionEntry* FunctionEntry)
+{
+	return MakeShareable(new FHeaderViewFunctionListItem(FunctionEntry));
+}
+
+FString FHeaderViewFunctionListItem::GetConditionalUFunctionSpecifiers(const UFunction* SigFunction) const
+{
+	check(SigFunction);
+
+	TStringBuilder<256> AdditionalSpecifiers;
+
+	if (SigFunction->FunctionFlags & FUNC_BlueprintPure)
+	{
+		AdditionalSpecifiers.Append(TEXT("BlueprintPure"));
+	}
+	else
+	{
+		AdditionalSpecifiers.Append(TEXT("BlueprintCallable"));
+	}
+
+	if (SigFunction->GetBoolMetaData(FBlueprintMetadata::MD_CallInEditor))
+	{
+		AdditionalSpecifiers.Append(TEXT(", CallInEditor"));
+	}
+
+	if (SigFunction->FunctionFlags & FUNC_Exec)
+	{
+		AdditionalSpecifiers.Append(TEXT(", Exec"));
+	}
+
+	const FString& Category = SigFunction->GetMetaData(FBlueprintMetadata::MD_FunctionCategory);
+	if (!Category.IsEmpty())
+	{
+		AdditionalSpecifiers.Append(FString::Printf(TEXT(", Category=\"%s\""), *Category));
+	}
+
+	// Meta Specifiers
+	{
+		TArray<FString> MetaSpecifiers;
+
+		if (SigFunction->GetBoolMetaData(FBlueprintMetadata::MD_ThreadSafe))
+		{
+			MetaSpecifiers.Emplace(TEXT("BlueprintThreadSafe"));
+		}
+
+		const FString& CompactNodeTitle = SigFunction->GetMetaData(FBlueprintMetadata::MD_CompactNodeTitle);
+		if (!CompactNodeTitle.IsEmpty())
+		{
+			MetaSpecifiers.Emplace(FString::Printf(TEXT("CompactNodeTitle=\"%s\""), *CompactNodeTitle));
+		}
+
+		const FString& DisplayName = SigFunction->GetMetaData(FBlueprintMetadata::MD_DisplayName);
+		if (!DisplayName.IsEmpty())
+		{
+			MetaSpecifiers.Emplace(FString::Printf(TEXT("DisplayName=\"%s\""), *DisplayName));
+		}
+		
+		const FString& Keywords = SigFunction->GetMetaData(FBlueprintMetadata::MD_FunctionKeywords);
+		if (!Keywords.IsEmpty())
+		{
+			MetaSpecifiers.Emplace(FString::Printf(TEXT("Keywords=\"%s\""), *Keywords));
+		}
+
+		if (!MetaSpecifiers.IsEmpty())
+		{
+			AdditionalSpecifiers.Append(FString::Printf(TEXT(", meta=(%s)"), *FString::Join(MetaSpecifiers, TEXT(", "))));
+		}
+	}
+
+	return AdditionalSpecifiers.ToString();
+}
+
+void FHeaderViewFunctionListItem::AppendFunctionParameters(const UFunction* SignatureFunction)
+{
+	check(SignatureFunction);
+
+	TArray<FString> Parameters;
+
+	// cache the return property to check against later
+	const FProperty* ReturnProperty = SignatureFunction->GetReturnProperty();
+
+	int32 ParamIdx = 0;
+	for (TFieldIterator<const FProperty> ParmIt(SignatureFunction); ParmIt; ++ParmIt)
+	{
+		// ReturnValue property shouldn't be duplicated in the parameter list
+		if (*ParmIt == ReturnProperty)
+		{
+			continue;
+		}
+
+		if (ParamIdx > 0)
+		{
+			RawItemString.Append(TEXT(", "));
+			RichTextString.Append(TEXT(", "));
+		}
+
+		if (ParmIt->HasAnyPropertyFlags(CPF_ConstParm))
+		{
+			RawItemString.Append(TEXT("const "));
+			RichTextString.Append(FString::Printf(TEXT("<%s>const</> "), *HeaderViewSyntaxDecorators::KeywordDecorator));
+		}
+		// If a param is declared as const&, then it is already treated as input and UPARAM(ref) will just be clutter
+		else if (ParmIt->HasAnyPropertyFlags(CPF_ReferenceParm))
+		{
+			RawItemString.Append(TEXT("UPARAM(ref) "));
+			RichTextString.Append(FString::Printf(TEXT("<%s>UPARAM</>(ref) "), *HeaderViewSyntaxDecorators::MacroDecorator));
+		}
+
+		if (ParmIt->HasAnyPropertyFlags(CPF_OutParm | CPF_ReferenceParm))
+		{
+			RawItemString.Append(FString::Printf(TEXT("%s& %s"), *ParmIt->GetCPPType(), *ParmIt->GetAuthoredName()));
+			RichTextString.Append(FString::Printf(TEXT("<%s>%s</>& <%s>%s</>"), *HeaderViewSyntaxDecorators::TypenameDecorator, *ParmIt->GetCPPType(), *HeaderViewSyntaxDecorators::IdentifierDecorator, *ParmIt->GetAuthoredName()));
+		}
+		else
+		{
+			RawItemString.Append(FString::Printf(TEXT("%s %s"), *ParmIt->GetCPPType(), *ParmIt->GetAuthoredName()));
+			RichTextString.Append(FString::Printf(TEXT("<%s>%s</> <%s>%s</>"), *HeaderViewSyntaxDecorators::TypenameDecorator, *ParmIt->GetCPPType(), *HeaderViewSyntaxDecorators::IdentifierDecorator, *ParmIt->GetAuthoredName()));
+		}
+
+		++ParamIdx;
+	}
+}
+
+FHeaderViewFunctionListItem::FHeaderViewFunctionListItem(const UK2Node_FunctionEntry* FunctionEntry)
+{
+	check(FunctionEntry);
+
+	RawItemString.Reserve(512);
+	RichTextString.Reserve(512);
+
+	if (const UFunction* ResolvedFunc = FunctionEntry->FindSignatureFunction())
+	{
+		// Format comment
+		{
+			FString Comment = ResolvedFunc->GetMetaData(FBlueprintMetadata::MD_Tooltip);
+			if (Comment.IsEmpty())
+			{
+				Comment = TEXT("Please add a function description");
+			}
+
+			if (ResolvedFunc->HasAnyFunctionFlags(FUNC_Event))
+			{
+				const UClass* OriginClass = ResolvedFunc->GetOwnerClass();
+				const FString OriginClassName = OriginClass->GetPrefixCPP() + OriginClass->GetName();
+				const FString EventType = ResolvedFunc->HasAnyFunctionFlags(FUNC_Native) ? TEXT("BlueprintNativeEvent") : TEXT("BlueprintImplementableEvent");
+				Comment.Append(FString::Printf(TEXT("\n\nNOTE: This function is linked to %s: %s::%s"), *EventType, *OriginClassName, *ResolvedFunc->GetName()));
+			}
+
+			FormatCommentString(Comment, RawItemString, RichTextString);
+		}
+
+		// Add Deprecation message if present
+		if (ResolvedFunc->GetBoolMetaData(FBlueprintMetadata::MD_DeprecatedFunction))
+		{
+			FString DeprecationMessage = ResolvedFunc->GetMetaData(FBlueprintMetadata::MD_DeprecationMessage);
+			if (DeprecationMessage.IsEmpty())
+			{
+				DeprecationMessage = TEXT("Please add a deprecation message.");
+			}
+
+			const FString EngineVersionString = FEngineVersion::Current().ToString(EVersionComponent::Patch);
+
+			RawItemString += FString::Printf(TEXT("\nUE_DEPRECATED(%s, \"%s\")"), *EngineVersionString, *DeprecationMessage);
+			RichTextString += FString::Printf(TEXT("\n<%s>UE_DEPRECATED</>(%s, \"%s\")"), *HeaderViewSyntaxDecorators::MacroDecorator, *EngineVersionString, *DeprecationMessage);
+		}
+
+		// Add the UFUNCTION specifiers
+		// i.e. UFUNCTION(BlueprintCallable, Category="Function Category")
+		{
+			const FString AdditionalSpecifiers = GetConditionalUFunctionSpecifiers(ResolvedFunc);
+
+			// Always add Blueprintable and BlueprintType
+			RawItemString += FString::Printf(TEXT("\nUFUNCTION(%s)"), *AdditionalSpecifiers);
+			RichTextString += FString::Printf(TEXT("\n<%s>UFUNCTION</>(%s)"), *HeaderViewSyntaxDecorators::MacroDecorator, *AdditionalSpecifiers);
+		}
+
+		// Add the function declaration line
+		// i.e. void FunctionName(Type InParam1, UPARAM(ref) Type2& InParam2, Type3& OutParam1)
+		{
+			FString FunctionName;
+			if (FunctionEntry->CustomGeneratedFunctionName.IsNone())
+			{
+				FunctionName = ResolvedFunc->GetName();
+			}
+			else
+			{
+				FunctionName = FunctionEntry->CustomGeneratedFunctionName.ToString();
+			}
+
+			if (FProperty* ReturnProperty = ResolvedFunc->GetReturnProperty())
+			{
+				const FString Typename = ReturnProperty->GetCPPType();
+				RawItemString += FString::Printf(TEXT("\n%s %s("), *Typename, *FunctionName);
+				RichTextString += FString::Printf(TEXT("\n<%s>%s</> <%s>%s</>("), *HeaderViewSyntaxDecorators::TypenameDecorator, *Typename, *HeaderViewSyntaxDecorators::IdentifierDecorator, *FunctionName);
+			}
+			else
+			{
+				RawItemString += FString::Printf(TEXT("\nvoid %s("), *FunctionName);
+				RichTextString += FString::Printf(TEXT("\n<%s>void</> <%s>%s</>("), *HeaderViewSyntaxDecorators::KeywordDecorator, *HeaderViewSyntaxDecorators::IdentifierDecorator, *FunctionName);
+			}
+
+
+			AppendFunctionParameters(ResolvedFunc);
+
+			if (FunctionEntry->GetFunctionFlags() & FUNC_Const)
+			{
+				RawItemString += TEXT(") const;");
+				RichTextString += FString::Printf(TEXT(") <%s>const</>;"), *HeaderViewSyntaxDecorators::KeywordDecorator);
+			}
+			else
+			{
+				RawItemString += TEXT(");");
+				RichTextString += TEXT(");");
+			}
+		}
+
+		// indent item
+		RawItemString.InsertAt(0, TEXT("\t"));
+		RichTextString.InsertAt(0, TEXT("\t"));
+		RawItemString.ReplaceInline(TEXT("\n"), TEXT("\n\t"));
+		RichTextString.ReplaceInline(TEXT("\n"), TEXT("\n\t"));
+
+		// normalize to platform newlines
+		UE::String::ToHostLineEndingsInline(RawItemString);
+		UE::String::ToHostLineEndingsInline(RichTextString);
+	}
+}
diff --git a/Engine/Plugins/Editor/BlueprintHeaderView/Source/BlueprintHeaderView/Private/HeaderViewFunctionListItem.h b/Engine/Plugins/Editor/BlueprintHeaderView/Source/BlueprintHeaderView/Private/HeaderViewFunctionListItem.h
new file mode 100644
--- /dev/null
+++ b/Engine/Plugins/Editor/BlueprintHeaderView/Source/BlueprintHeaderView/Private/HeaderViewFunctionListItem.h
@@ -0,0 +1,25 @@
+// Copyright Epic Games, Inc. All Rights Reserved.
+
+#pragma once
+
+#include "CoreMinimal.h"
+#include "SBlueprintHeaderView.h"
+
+class UK2Node_FunctionEntry;
+
+/** A header view list item that displays a function declaration */
+struct FHeaderViewFunctionListItem : public FHeaderViewListItem
+{
+	/** Creates a list item for the Header view representing a function declaration for the given blueprint function */
+	static FHeaderViewListItemPtr Create(const UK2Node_FunctionEntry* FunctionEntry);
+
+protected:
+	FHeaderViewFunctionListItem(const UK2Node_FunctionEntry* FunctionEntry);
+
+	/** Returns a string containing the specifiers for the UFUNCTION line */
+	FString GetConditionalUFunctionSpecifiers(const UFunction* SigFunction) const;
+
+	/** Adds Function parameters to the RichText and PlainText strings */
+	void AppendFunctionParameters(const UFunction* SignatureFunction);
+
+};
\ No newline at end of file
diff --git a/Engine/Plugins/Editor/BlueprintHeaderView/Source/BlueprintHeaderView/Private/SBlueprintHeaderView.cpp b/Engine/Plugins/Editor/BlueprintHeaderView/Source/BlueprintHeaderView/Private/SBlueprintHeaderView.cpp
--- a/Engine/Plugins/Editor/BlueprintHeaderView/Source/BlueprintHeaderView/Private/SBlueprintHeaderView.cpp
+++ b/Engine/Plugins/Editor/BlueprintHeaderView/Source/BlueprintHeaderView/Private/SBlueprintHeaderView.cpp
@@ -11,11 +11,25 @@
 #include "Engine/Blueprint.h"
 #include "Framework/Text/SlateTextRun.h"
 #include "HeaderViewClassListItem.h"
+#include "HeaderViewFunctionListItem.h"
+#include "EdGraphSchema_K2.h"
+#include "K2Node_FunctionEntry.h"
 #include "EditorStyleSet.h"
 #include "String/LineEndings.h"
 
 #define LOCTEXT_NAMESPACE "SBlueprintHeaderView"
 
+// HeaderViewSyntaxDecorators /////////////////////////////////////////////////
+
+namespace HeaderViewSyntaxDecorators
+{
+	const FString CommentDecorator = TEXT("comment");
+	const FString IdentifierDecorator = TEXT("identifier");
+	const FString KeywordDecorator = TEXT("keyword");
+	const FString MacroDecorator = TEXT("macro");
+	const FString TypenameDecorator = TEXT("typename");
+}
+
 // FHeaderViewSyntaxDecorator /////////////////////////////////////////////////
 
 namespace
@@ -89,11 +103,11 @@ TSharedRef<SWidget> FHeaderViewListItem::GenerateWidgetForItem()
 			SNew(SRichTextBlock)
 			.Text(FText::FromString(RichTextString))
 			.TextStyle(FEditorStyle::Get(), "Log.Normal")
-			+SRichTextBlock::Decorator(FHeaderViewSyntaxDecorator::Create(TEXT("comment"), SyntaxColors.Comment))
-			+SRichTextBlock::Decorator(FHeaderViewSyntaxDecorator::Create(TEXT("macro"), SyntaxColors.Macro))
-			+SRichTextBlock::Decorator(FHeaderViewSyntaxDecorator::Create(TEXT("typename"), SyntaxColors.Typename))
-			+SRichTextBlock::Decorator(FHeaderViewSyntaxDecorator::Create(TEXT("identifier"), SyntaxColors.Identifier))
-			+SRichTextBlock::Decorator(FHeaderViewSyntaxDecorator::Create(TEXT("keyword"), SyntaxColors.Keyword))
+			+SRichTextBlock::Decorator(FHeaderViewSyntaxDecorator::Create(HeaderViewSyntaxDecorators::CommentDecorator, SyntaxColors.Comment))
+			+SRichTextBlock::Decorator(FHeaderViewSyntaxDecorator::Create(HeaderViewSyntaxDecorators::IdentifierDecorator, SyntaxColors.Identifier))
+			+SRichTextBlock::Decorator(FHeaderViewSyntaxDecorator::Create(HeaderViewSyntaxDecorators::KeywordDecorator, SyntaxColors.Keyword))
+			+SRichTextBlock::Decorator(FHeaderViewSyntaxDecorator::Create(HeaderViewSyntaxDecorators::MacroDecorator, SyntaxColors.Macro))
+			+SRichTextBlock::Decorator(FHeaderViewSyntaxDecorator::Create(HeaderViewSyntaxDecorators::TypenameDecorator, SyntaxColors.Typename))
 		];
 }
 
@@ -132,8 +146,8 @@ void FHeaderViewListItem::FormatCommentString(FString InComment, FString& OutRaw
 	OutRawString = InComment;
 
 	// mark each line of the comment as the beginning and end of a comment style for the rich text representation
-	InComment.ReplaceInline(TEXT("\n"), TEXT("</>\n<comment>"));
-	OutRichString = FString::Printf(TEXT("<comment>%s</>"), *InComment);
+	InComment.ReplaceInline(TEXT("\n"), *FString::Printf(TEXT("</>\n<%s>"), *HeaderViewSyntaxDecorators::CommentDecorator));
+	OutRichString = FString::Printf(TEXT("<%s>%s</>"), *HeaderViewSyntaxDecorators::CommentDecorator, *InComment);
 }
 
 // SBlueprintHeaderView ///////////////////////////////////////////////////////
@@ -249,14 +263,65 @@ void SBlueprintHeaderView::RepopulateListView()
 {
 	ListItems.Empty();
 
-	// Add the class declaration
-	ListItems.Add(FHeaderViewClassListItem::Create(SelectedBlueprint));
+	if (UBlueprint* Blueprint = SelectedBlueprint.Get())
+	{
+		// Add the class declaration
+		ListItems.Add(FHeaderViewClassListItem::Create(SelectedBlueprint));
+
+		PopulateFunctionItems(Blueprint);
 
+		// Add the closing brace of the class
+		ListItems.Add(FHeaderViewListItem::Create(TEXT("};"), TEXT("};")));
+	}
 
-	// Add the closing brace of the class
-	ListItems.Add(FHeaderViewListItem::Create(TEXT("};"), TEXT("};")));
-	
 	ListView->RequestListRefresh();
 }
 
+void SBlueprintHeaderView::PopulateFunctionItems(const UBlueprint* Blueprint)
+{
+	if (Blueprint)
+	{
+		// We should only add an access specifier line if the previous function was a different one
+		int32 PrevAccessSpecifier = 0;
+		for (const UEdGraph* FunctionGraph : Blueprint->FunctionGraphs)
+		{
+			if (FunctionGraph && !UEdGraphSchema_K2::IsConstructionScript(FunctionGraph))
+			{
+				TArray<UK2Node_FunctionEntry*> EntryNodes;
+				FunctionGraph->GetNodesOfClass<UK2Node_FunctionEntry>(EntryNodes);
+
+				if (ensure(EntryNodes.Num() == 1 && EntryNodes[0]))
+				{
+					int32 AccessSpecifier = EntryNodes[0]->GetFunctionFlags() & FUNC_AccessSpecifiers;
+
+					if (AccessSpecifier != PrevAccessSpecifier)
+					{
+						switch (AccessSpecifier)
+						{
+						case FUNC_Public:
+							ListItems.Add(FHeaderViewListItem::Create(TEXT("public:"), FString::Printf(TEXT("<%s>public</>:"), *HeaderViewSyntaxDecorators::KeywordDecorator)));
+							break;
+						case FUNC_Protected:
+							ListItems.Add(FHeaderViewListItem::Create(TEXT("protected:"), FString::Printf(TEXT("<%s>protected</>:"), *HeaderViewSyntaxDecorators::KeywordDecorator)));
+							break;
+						case FUNC_Private:
+							ListItems.Add(FHeaderViewListItem::Create(TEXT("private:"), FString::Printf(TEXT("<%s>private</>:"), *HeaderViewSyntaxDecorators::KeywordDecorator)));
+							break;
+						}
+					}
+					else
+					{
+						// add an empty line to space functions out
+						ListItems.Add(FHeaderViewListItem::Create(TEXT(""), TEXT("")));
+					}
+
+					PrevAccessSpecifier = AccessSpecifier;
+
+					ListItems.Add(FHeaderViewFunctionListItem::Create(EntryNodes[0]));
+				}
+			}
+		}
+	}
+}
+
 #undef LOCTEXT_NAMESPACE
diff --git a/Engine/Plugins/Editor/BlueprintHeaderView/Source/BlueprintHeaderView/Private/SBlueprintHeaderView.h b/Engine/Plugins/Editor/BlueprintHeaderView/Source/BlueprintHeaderView/Private/SBlueprintHeaderView.h
--- a/Engine/Plugins/Editor/BlueprintHeaderView/Source/BlueprintHeaderView/Private/SBlueprintHeaderView.h
+++ b/Engine/Plugins/Editor/BlueprintHeaderView/Source/BlueprintHeaderView/Private/SBlueprintHeaderView.h
@@ -13,6 +13,16 @@ class UBlueprint;
 class ITableRow;
 class STableViewBase;
 
+/** rich text decorators for BlueprintHeaderView Syntax Highlighting */
+namespace HeaderViewSyntaxDecorators
+{
+	extern const FString CommentDecorator;
+	extern const FString IdentifierDecorator;
+	extern const FString KeywordDecorator;
+	extern const FString MacroDecorator;
+	extern const FString TypenameDecorator;
+}
+
 /** A base class for List Items in the Header View */
 struct FHeaderViewListItem : public TSharedFromThis<FHeaderViewListItem>
 {
@@ -73,6 +83,9 @@ private:
 	
 	/** Clears the list and repopulates it with info for the selected Blueprint */
 	void RepopulateListView();
+
+	/** Adds items to the list view representing all functions present in the given blueprint */
+	void PopulateFunctionItems(const UBlueprint* Blueprint);
 private:
 	/** The blueprint currently being displayed by the header view */
 	TWeakObjectPtr<UBlueprint> SelectedBlueprint;
