BlueprintHeaderView: Add ListView items for variables and fix function parameter type names

diff --git a/Engine/Plugins/Editor/BlueprintHeaderView/Source/BlueprintHeaderView/Private/HeaderViewFunctionListItem.cpp b/Engine/Plugins/Editor/BlueprintHeaderView/Source/BlueprintHeaderView/Private/HeaderViewFunctionListItem.cpp
--- a/Engine/Plugins/Editor/BlueprintHeaderView/Source/BlueprintHeaderView/Private/HeaderViewFunctionListItem.cpp
+++ b/Engine/Plugins/Editor/BlueprintHeaderView/Source/BlueprintHeaderView/Private/HeaderViewFunctionListItem.cpp
@@ -117,8 +117,10 @@ void FHeaderViewFunctionListItem::AppendFunctionParameters(const UFunction* Sign
 
 		if (ParmIt->HasAnyPropertyFlags(CPF_OutParm | CPF_ReferenceParm))
 		{
-			RawItemString.Append(FString::Printf(TEXT("%s& %s"), *ParmIt->GetCPPType(), *ParmIt->GetAuthoredName()));
-			RichTextString.Append(FString::Printf(TEXT("<%s>%s</>& <%s>%s</>"), *HeaderViewSyntaxDecorators::TypenameDecorator, *ParmIt->GetCPPType(), *HeaderViewSyntaxDecorators::IdentifierDecorator, *ParmIt->GetAuthoredName()));
+			const FString Typename = GetCPPTypenameForProperty(*ParmIt);
+			const FString ParmName = ParmIt->GetAuthoredName();
+			RawItemString.Append(FString::Printf(TEXT("%s& %s"), *Typename, *ParmName));
+			RichTextString.Append(FString::Printf(TEXT("<%s>%s</>& <%s>%s</>"), *HeaderViewSyntaxDecorators::TypenameDecorator, *Typename, *HeaderViewSyntaxDecorators::IdentifierDecorator, *ParmName));
 		}
 		else
 		{
diff --git a/Engine/Plugins/Editor/BlueprintHeaderView/Source/BlueprintHeaderView/Private/HeaderViewVariableListItem.cpp b/Engine/Plugins/Editor/BlueprintHeaderView/Source/BlueprintHeaderView/Private/HeaderViewVariableListItem.cpp
new file mode 100644
--- /dev/null
+++ b/Engine/Plugins/Editor/BlueprintHeaderView/Source/BlueprintHeaderView/Private/HeaderViewVariableListItem.cpp
@@ -0,0 +1,232 @@
+// Copyright Epic Games, Inc. All Rights Reserved.
+
+#include "HeaderViewVariableListItem.h"
+#include "EdMode.h"
+#include "EdGraphSchema_K2.h"
+#include "Engine/Blueprint.h"
+#include "String/LineEndings.h"
+
+FHeaderViewListItemPtr FHeaderViewVariableListItem::Create(const FBPVariableDescription& VariableDesc, const FProperty& VarProperty)
+{
+	return MakeShareable(new FHeaderViewVariableListItem(VariableDesc, VarProperty));
+}
+
+FHeaderViewVariableListItem::FHeaderViewVariableListItem(const FBPVariableDescription& VariableDesc, const FProperty& VarProperty)
+{
+	// Format comment
+	{
+		FString Comment = VarProperty.GetMetaData(FBlueprintMetadata::MD_Tooltip);
+		if (Comment.IsEmpty())
+		{
+			Comment = TEXT("Please add a variable description");
+		}
+
+		FormatCommentString(Comment, RawItemString, RichTextString);
+	}
+
+	// Add a static_assert if the user needs to tell these variables how to replicate
+	// i.e. static_assert(false, "You will need to add DOREPLIFETIME(ClassName, VarName) to GetLifeTimeReplicatedProps");
+	if (VarProperty.HasAnyPropertyFlags(CPF_Net))
+	{
+		FString ClassName = GetOwningClassName(VarProperty);
+		if (VariableDesc.ReplicationCondition == COND_None)
+		{
+			RawItemString += FString::Printf(TEXT("\nstatic_assert(false, \"You will need to add DOREPLIFETIME(%s, %s) to GetLifetimeReplicatedProps\");"), *ClassName, *VarProperty.GetAuthoredName());
+			RichTextString += FString::Printf(TEXT("\n<%s>static_assert</>(<%s>false</>, \"You will need to add DOREPLIFETIME(%s, %s) to GetLifetimeReplicatedProps\");"), 
+				*HeaderViewSyntaxDecorators::KeywordDecorator, 
+				*HeaderViewSyntaxDecorators::KeywordDecorator, 
+				*ClassName, 
+				*VarProperty.GetAuthoredName());
+		}
+		else
+		{
+			FString ConditionString = StaticEnum<ELifetimeCondition>()->GetAuthoredNameStringByValue(VariableDesc.ReplicationCondition);
+			RawItemString += FString::Printf(TEXT("\nstatic_assert(false, \"You will need to add DOREPLIFETIME(%s, %s) to GetLifetimeReplicatedProps\");"), *ClassName, *VarProperty.GetAuthoredName());
+			RichTextString += FString::Printf(TEXT("\n<%s>static_assert</>(<%s>false</>, \"You will need to add DOREPLIFETIME_WITH_PARAMS(%s, %s, %s) to GetLifetimeReplicatedProps\");"),
+				*HeaderViewSyntaxDecorators::KeywordDecorator,
+				*HeaderViewSyntaxDecorators::KeywordDecorator,
+				*ClassName,
+				*VarProperty.GetAuthoredName(),
+				*ConditionString);
+		}
+	}
+
+	// Add Deprecation message if present
+	if (VarProperty.HasAnyPropertyFlags(CPF_Deprecated))
+	{
+		FString DeprecationMessage = VarProperty.GetMetaData(FBlueprintMetadata::MD_DeprecationMessage);
+		if (DeprecationMessage.IsEmpty())
+		{
+			DeprecationMessage = TEXT("Please add a deprecation message.");
+		}
+
+		const FString EngineVersionString = FEngineVersion::Current().ToString(EVersionComponent::Patch);
+
+		RawItemString += FString::Printf(TEXT("\nUE_DEPRECATED(%s, \"%s\")"), *EngineVersionString, *DeprecationMessage);
+		RichTextString += FString::Printf(TEXT("\n<%s>UE_DEPRECATED</>(%s, \"%s\")"), *HeaderViewSyntaxDecorators::MacroDecorator, *EngineVersionString, *DeprecationMessage);
+	}
+
+	// Add the UPROPERTY specifiers
+	// i.e. UPROPERTY(BlueprintReadWrite, Category="Variable Category")
+	{
+		const FString AdditionalSpecifiers = GetConditionalUPropertySpecifiers(VarProperty);
+
+		// Always add Blueprintable and BlueprintType
+		RawItemString += FString::Printf(TEXT("\nUPROPERTY(%s)"), *AdditionalSpecifiers);
+		RichTextString += FString::Printf(TEXT("\n<%s>UPROPERTY</>(%s)"), *HeaderViewSyntaxDecorators::MacroDecorator, *AdditionalSpecifiers);
+	}
+
+	// Add the variable declaration line
+	// i.e. Type VariableName;
+	{
+		const FString Typename = GetCPPTypenameForProperty(&VarProperty);
+		const FString VarName = VarProperty.GetAuthoredName();
+
+		RawItemString += FString::Printf(TEXT("\n%s %s;"), *Typename, *VarName);
+		RichTextString += FString::Printf(TEXT("\n<%s>%s</> <%s>%s</>;"), 
+			*HeaderViewSyntaxDecorators::TypenameDecorator,
+			*Typename.Replace(TEXT("<"), TEXT("&lt;")).Replace(TEXT(">"), TEXT("&gt;")),
+			*HeaderViewSyntaxDecorators::IdentifierDecorator,
+			*VarName);
+	}
+
+	// indent item
+	RawItemString.InsertAt(0, TEXT("\t"));
+	RichTextString.InsertAt(0, TEXT("\t"));
+	RawItemString.ReplaceInline(TEXT("\n"), TEXT("\n\t"));
+	RichTextString.ReplaceInline(TEXT("\n"), TEXT("\n\t"));
+
+	// normalize to platform newlines
+	UE::String::ToHostLineEndingsInline(RawItemString);
+	UE::String::ToHostLineEndingsInline(RichTextString);
+}
+
+FString FHeaderViewVariableListItem::GetConditionalUPropertySpecifiers(const FProperty& VarProperty) const
+{
+	TArray<FString> PropertySpecifiers;
+
+	if (!VarProperty.HasMetaData(FBlueprintMetadata::MD_Private) || !VarProperty.GetBoolMetaData(FBlueprintMetadata::MD_Private))
+	{
+		if (VarProperty.HasAllPropertyFlags(CPF_BlueprintReadOnly))
+		{
+			PropertySpecifiers.Emplace(TEXT("BlueprintReadOnly"));
+		}
+		else
+		{
+			PropertySpecifiers.Emplace(TEXT("BlueprintReadWrite"));
+		}
+	}
+
+	if (!VarProperty.HasAnyPropertyFlags(CPF_DisableEditOnInstance | CPF_DisableEditOnTemplate))
+	{
+		PropertySpecifiers.Emplace(TEXT("EditAnywhere"));
+	}
+	else if (VarProperty.HasAnyPropertyFlags(CPF_DisableEditOnInstance))
+	{
+		PropertySpecifiers.Emplace(TEXT("EditDefaultsOnly"));
+	}
+	else
+	{
+		PropertySpecifiers.Emplace(TEXT("EditInstanceOnly"));
+	}
+
+	if (VarProperty.HasMetaData(FBlueprintMetadata::MD_FunctionCategory))
+	{
+		PropertySpecifiers.Emplace(FString::Printf(TEXT("Category=\"%s\""), *VarProperty.GetMetaData(FBlueprintMetadata::MD_FunctionCategory)));
+	}
+
+	if (VarProperty.HasAnyPropertyFlags(CPF_Net))
+	{
+		if (VarProperty.HasAnyPropertyFlags(CPF_RepNotify))
+		{
+			PropertySpecifiers.Emplace(FString::Printf(TEXT("ReplicatedUsing=\"OnRep_%s\""), *VarProperty.GetAuthoredName()));
+		}
+		else
+		{
+			PropertySpecifiers.Emplace(TEXT("Replicated"));
+		}
+	}
+
+	if (VarProperty.HasAnyPropertyFlags(CPF_Interp))
+	{
+		PropertySpecifiers.Emplace(TEXT("Interp"));
+	}
+
+	if (VarProperty.HasAnyPropertyFlags(CPF_Config))
+	{
+		PropertySpecifiers.Emplace(TEXT("Config"));
+	}
+
+	if (VarProperty.HasAnyPropertyFlags(CPF_Transient))
+	{
+		PropertySpecifiers.Emplace(TEXT("Transient"));
+	}
+
+	if (VarProperty.HasAnyPropertyFlags(CPF_SaveGame))
+	{
+		PropertySpecifiers.Emplace(TEXT("SaveGame"));
+	}
+
+	if (VarProperty.HasAnyPropertyFlags(CPF_AdvancedDisplay))
+	{
+		PropertySpecifiers.Emplace(TEXT("AdvancedDisplay"));
+	}
+
+	if (VarProperty.HasMetaData(FEdMode::MD_MakeEditWidget) && VarProperty.GetBoolMetaData(FEdMode::MD_MakeEditWidget))
+	{
+		PropertySpecifiers.Emplace(TEXT("MakeEditWidget"));
+	}
+
+	// meta specifiers
+	{
+		// These metadata keys are handled explicitly elsewhere
+		static const TArray<FName> IgnoreMetaData = 
+		{
+			FBlueprintMetadata::MD_DeprecationMessage,
+			FBlueprintMetadata::MD_FunctionCategory,
+			FBlueprintMetadata::MD_Tooltip,
+			FBlueprintMetadata::MD_Private,
+			FEdMode::MD_MakeEditWidget
+		};
+
+		TArray<FString> MetaSpecifiers;
+
+		for (const TPair<FName, FString>& MetaPair : *VarProperty.GetMetaDataMap())
+		{
+			// Don't take this metadata if we handled it explicitly
+			if (IgnoreMetaData.Contains(MetaPair.Key))
+			{
+				continue;
+			}
+
+			// Don't add the DisplayName metadata key if the display name is just the friendly name generated by the editor anyway
+			if (MetaPair.Key == FBlueprintMetadata::MD_DisplayName)
+			{
+				if (MetaPair.Value == FName::NameToDisplayString(VarProperty.GetName(), !!CastField<FBoolProperty>(&VarProperty)))
+				{
+					continue;
+				}
+			}
+
+			MetaSpecifiers.Emplace(FString::Printf(TEXT("%s=\"%s\""), *MetaPair.Key.ToString(), *MetaPair.Value));
+		}
+
+		if (!MetaSpecifiers.IsEmpty())
+		{
+			PropertySpecifiers.Emplace(FString::Printf(TEXT("meta=(%s)"), *FString::Join(MetaSpecifiers, TEXT(", "))));
+		}
+	}
+
+	return FString::Join(PropertySpecifiers, TEXT(", "));
+}
+
+FString FHeaderViewVariableListItem::GetOwningClassName(const FProperty& VarProperty) const
+{
+	UClass* OwningClass = VarProperty.GetOwnerClass();
+	if (OwningClass && OwningClass->ClassGeneratedBy)
+	{
+		return OwningClass->GetPrefixCPP() + OwningClass->ClassGeneratedBy->GetName();
+	}
+
+	return TEXT("");
+}
diff --git a/Engine/Plugins/Editor/BlueprintHeaderView/Source/BlueprintHeaderView/Private/HeaderViewVariableListItem.h b/Engine/Plugins/Editor/BlueprintHeaderView/Source/BlueprintHeaderView/Private/HeaderViewVariableListItem.h
new file mode 100644
--- /dev/null
+++ b/Engine/Plugins/Editor/BlueprintHeaderView/Source/BlueprintHeaderView/Private/HeaderViewVariableListItem.h
@@ -0,0 +1,25 @@
+// Copyright Epic Games, Inc. All Rights Reserved.
+
+#pragma once
+
+#include "CoreMinimal.h"
+#include "SBlueprintHeaderView.h"
+
+struct FBPVariableDescription;
+
+/** A header view list item that displays a variable declaration */
+struct FHeaderViewVariableListItem : public FHeaderViewListItem
+{
+	/** Creates a list item for the Header view representing a variable declaration for the given blueprint variable */
+	static FHeaderViewListItemPtr Create(const FBPVariableDescription& VariableDesc, const FProperty& VarProperty);
+
+protected:
+	FHeaderViewVariableListItem(const FBPVariableDescription& VariableDesc, const FProperty& VarProperty);
+
+	/** Returns a string containing the specifiers for the UPROPERTY line */
+	FString GetConditionalUPropertySpecifiers(const FProperty& VarProperty) const;
+
+	/** Returns the name of the owning class */
+	FString GetOwningClassName(const FProperty& VarProperty) const;
+
+};
\ No newline at end of file
diff --git a/Engine/Plugins/Editor/BlueprintHeaderView/Source/BlueprintHeaderView/Private/SBlueprintHeaderView.cpp b/Engine/Plugins/Editor/BlueprintHeaderView/Source/BlueprintHeaderView/Private/SBlueprintHeaderView.cpp
--- a/Engine/Plugins/Editor/BlueprintHeaderView/Source/BlueprintHeaderView/Private/SBlueprintHeaderView.cpp
+++ b/Engine/Plugins/Editor/BlueprintHeaderView/Source/BlueprintHeaderView/Private/SBlueprintHeaderView.cpp
@@ -12,6 +12,7 @@
 #include "Framework/Text/SlateTextRun.h"
 #include "HeaderViewClassListItem.h"
 #include "HeaderViewFunctionListItem.h"
+#include "HeaderViewVariableListItem.h"
 #include "EdGraphSchema_K2.h"
 #include "K2Node_FunctionEntry.h"
 #include "EditorStyleSet.h"
@@ -150,6 +151,19 @@ void FHeaderViewListItem::FormatCommentString(FString InComment, FString& OutRaw
 	OutRichString = FString::Printf(TEXT("<%s>%s</>"), *HeaderViewSyntaxDecorators::CommentDecorator, *InComment);
 }
 
+FString FHeaderViewListItem::GetCPPTypenameForProperty(const FProperty* InProperty)
+{
+	if (InProperty)
+	{
+		FString ExtendedTypeText;
+		return InProperty->GetCPPType(&ExtendedTypeText) + ExtendedTypeText;
+	}
+	else
+	{
+		return TEXT("");
+	}
+}
+
 // SBlueprintHeaderView ///////////////////////////////////////////////////////
 
 void SBlueprintHeaderView::Construct(const FArguments& InArgs)
@@ -270,6 +284,8 @@ void SBlueprintHeaderView::RepopulateListView()
 
 		PopulateFunctionItems(Blueprint);
 
+		PopulateVariableItems(Blueprint);
+
 		// Add the closing brace of the class
 		ListItems.Add(FHeaderViewListItem::Create(TEXT("};"), TEXT("};")));
 	}
@@ -298,6 +314,7 @@ void SBlueprintHeaderView::PopulateFunctionItems(const UBlueprint* Blueprint)
 					{
 						switch (AccessSpecifier)
 						{
+						default: //< Blueprint OnRep functions don't have one of these flags set, default to public in this case and others like it
 						case FUNC_Public:
 							ListItems.Add(FHeaderViewListItem::Create(TEXT("public:"), FString::Printf(TEXT("<%s>public</>:"), *HeaderViewSyntaxDecorators::KeywordDecorator)));
 							break;
@@ -324,4 +341,44 @@ void SBlueprintHeaderView::PopulateFunctionItems(const UBlueprint* Blueprint)
 	}
 }
 
+void SBlueprintHeaderView::PopulateVariableItems(const UBlueprint* Blueprint)
+{
+	if (Blueprint)
+	{
+		const int32 Private = 2;
+		const int32 Public = 1;
+
+		// We should only add an access specifier line if the previous variable was a different one
+		int32 PrevAccessSpecifier = 0;
+		for (const FBPVariableDescription& VariableDesc : Blueprint->NewVariables)
+		{
+			if (const FProperty* VarProperty = Blueprint->SkeletonGeneratedClass->FindPropertyByName(VariableDesc.VarName))
+			{
+				const int32 AccessSpecifier = VarProperty->GetBoolMetaData(FBlueprintMetadata::MD_Private) ? Private : Public;
+				if (AccessSpecifier != PrevAccessSpecifier)
+				{
+					switch (AccessSpecifier)
+					{
+					case Public:
+						ListItems.Add(FHeaderViewListItem::Create(TEXT("public:"), FString::Printf(TEXT("<%s>public</>:"), *HeaderViewSyntaxDecorators::KeywordDecorator)));
+						break;
+					case Private:
+						ListItems.Add(FHeaderViewListItem::Create(TEXT("private:"), FString::Printf(TEXT("<%s>private</>:"), *HeaderViewSyntaxDecorators::KeywordDecorator)));
+						break;
+					}
+
+					PrevAccessSpecifier = AccessSpecifier;
+				}
+				else
+				{
+					// add an empty line to space variables out
+					ListItems.Add(FHeaderViewListItem::Create(TEXT(""), TEXT("")));
+				}
+
+				ListItems.Add(FHeaderViewVariableListItem::Create(VariableDesc, *VarProperty));
+			}
+		}
+	}
+}
+
 #undef LOCTEXT_NAMESPACE
diff --git a/Engine/Plugins/Editor/BlueprintHeaderView/Source/BlueprintHeaderView/Private/SBlueprintHeaderView.h b/Engine/Plugins/Editor/BlueprintHeaderView/Source/BlueprintHeaderView/Private/SBlueprintHeaderView.h
--- a/Engine/Plugins/Editor/BlueprintHeaderView/Source/BlueprintHeaderView/Private/SBlueprintHeaderView.h
+++ b/Engine/Plugins/Editor/BlueprintHeaderView/Source/BlueprintHeaderView/Private/SBlueprintHeaderView.h
@@ -46,6 +46,12 @@ protected:
 	 */
 	static void FormatCommentString(FString InComment, FString& OutRawString, FString& OutRichString);
 
+	/** 
+	 * returns a string representing the full C++ typename for the given property, 
+	 * including template params for container types
+	 */
+	static FString GetCPPTypenameForProperty(const FProperty* InProperty);
+
 protected:
 	/** A rich text representation of the item, including syntax highlighting and errors */
 	FString RichTextString;
@@ -86,6 +92,9 @@ private:
 
 	/** Adds items to the list view representing all functions present in the given blueprint */
 	void PopulateFunctionItems(const UBlueprint* Blueprint);
+
+	/** Adds items to the list view representing all variables present in the given blueprints */
+	void PopulateVariableItems(const UBlueprint* Blueprint);
 private:
 	/** The blueprint currently being displayed by the header view */
 	TWeakObjectPtr<UBlueprint> SelectedBlueprint;
diff --git a/Engine/Source/Runtime/CoreUObject/Private/UObject/Class.cpp b/Engine/Source/Runtime/CoreUObject/Private/UObject/Class.cpp
--- a/Engine/Source/Runtime/CoreUObject/Private/UObject/Class.cpp
+++ b/Engine/Source/Runtime/CoreUObject/Private/UObject/Class.cpp
@@ -3266,7 +3266,7 @@ FGuid UScriptStruct::GetCustomGuid() const
 	return FGuid();
 }
 
-FString UScriptStruct::GetStructCPPName() const
+FString UScriptStruct::GetStructCPPName(uint32 CPPExportFlags/*=0*/) const
 {
 	return FString::Printf(TEXT("F%s"), *GetName());
 }
diff --git a/Engine/Source/Runtime/CoreUObject/Private/UObject/PropertyStruct.cpp b/Engine/Source/Runtime/CoreUObject/Private/UObject/PropertyStruct.cpp
--- a/Engine/Source/Runtime/CoreUObject/Private/UObject/PropertyStruct.cpp
+++ b/Engine/Source/Runtime/CoreUObject/Private/UObject/PropertyStruct.cpp
@@ -266,7 +266,7 @@ bool FStructProperty::HasNoOpConstructor() const
 
 FString FStructProperty::GetCPPType( FString* ExtendedTypeText/*=NULL*/, uint32 CPPExportFlags/*=0*/ ) const
 {
-	return Struct->GetStructCPPName();
+	return Struct->GetStructCPPName(CPPExportFlags);
 }
 
 FString FStructProperty::GetCPPTypeForwardDeclaration() const
diff --git a/Engine/Source/Runtime/CoreUObject/Public/UObject/Class.h b/Engine/Source/Runtime/CoreUObject/Public/UObject/Class.h
--- a/Engine/Source/Runtime/CoreUObject/Public/UObject/Class.h
+++ b/Engine/Source/Runtime/CoreUObject/Public/UObject/Class.h
@@ -1869,7 +1869,7 @@ public:
 	virtual COREUOBJECT_API FGuid GetCustomGuid() const;
 	
 	/** Returns the (native, c++) name of the struct */
-	virtual COREUOBJECT_API FString GetStructCPPName() const;
+	virtual COREUOBJECT_API FString GetStructCPPName(uint32 CPPExportFlags = 0) const;
 
 	/**
 	 * Initializes this structure to its default values
diff --git a/Engine/Source/Runtime/Engine/Classes/Engine/UserDefinedStruct.h b/Engine/Source/Runtime/Engine/Classes/Engine/UserDefinedStruct.h
--- a/Engine/Source/Runtime/Engine/Classes/Engine/UserDefinedStruct.h
+++ b/Engine/Source/Runtime/Engine/Classes/Engine/UserDefinedStruct.h
@@ -105,7 +105,7 @@ public:
 	virtual uint32 GetStructTypeHash(const void* Src) const override;
 	virtual void RecursivelyPreload() override;
 	virtual FGuid GetCustomGuid() const override;
-	virtual FString GetStructCPPName() const override;
+	virtual FString GetStructCPPName(uint32 CPPExportFlags) const override;
 	virtual FProperty* CustomFindProperty(const FName Name) const override;
 	virtual void PrepareCppStructOps();
 	// End of  UScriptStruct interface.
diff --git a/Engine/Source/Runtime/Engine/Private/UserDefinedStruct.cpp b/Engine/Source/Runtime/Engine/Private/UserDefinedStruct.cpp
--- a/Engine/Source/Runtime/Engine/Private/UserDefinedStruct.cpp
+++ b/Engine/Source/Runtime/Engine/Private/UserDefinedStruct.cpp
@@ -356,9 +356,16 @@ ENGINE_API FString GetPathPostfix(const UObject* ForObject)
 	return FString::Printf(TEXT("%u"), GetTypeHash(FullAssetName));
 }
 
-FString UUserDefinedStruct::GetStructCPPName() const
+FString UUserDefinedStruct::GetStructCPPName(uint32 CPPExportFlags) const
 {
-	return ::UnicodeToCPPIdentifier(*GetName(), false, GetPrefixCPP()) + GetPathPostfix(this);
+	if (CPPExportFlags & CPPF_BlueprintCppBackend)
+	{
+		return ::UnicodeToCPPIdentifier(*GetName(), false, GetPrefixCPP()) + GetPathPostfix(this);
+	}
+	else
+	{
+		return Super::GetStructCPPName(CPPExportFlags);
+	}
 }
 
 uint32 UUserDefinedStruct::GetUserDefinedStructTypeHash(const void* Src, const UScriptStruct* Type)
