BPHeaderView: Fix typo in string literal

#rnx

diff --git a/Engine/Plugins/Editor/BlueprintHeaderView/Source/BlueprintHeaderView/Private/HeaderViewVariableListItem.cpp b/Engine/Plugins/Editor/BlueprintHeaderView/Source/BlueprintHeaderView/Private/HeaderViewVariableListItem.cpp
--- a/Engine/Plugins/Editor/BlueprintHeaderView/Source/BlueprintHeaderView/Private/HeaderViewVariableListItem.cpp
+++ b/Engine/Plugins/Editor/BlueprintHeaderView/Source/BlueprintHeaderView/Private/HeaderViewVariableListItem.cpp
@@ -41,7 +41,7 @@ FHeaderViewVariableListItem::FHeaderViewVariableListItem(const FBPVariableDescri
 		else
 		{
 			FString ConditionString = StaticEnum<ELifetimeCondition>()->GetAuthoredNameStringByValue(VariableDesc.ReplicationCondition);
-			RawItemString += FString::Printf(TEXT("\nstatic_assert(false, \"You will need to add DOREPLIFETIME(%s, %s) to GetLifetimeReplicatedProps\");"), *ClassName, *VarProperty.GetAuthoredName());
+			RawItemString += FString::Printf(TEXT("\nstatic_assert(false, \"You will need to add DOREPLIFETIME_WITH_PARAMS(%s, %s, %s) to GetLifetimeReplicatedProps\");"), *ClassName, *VarProperty.GetAuthoredName(), *ConditionString);
 			RichTextString += FString::Printf(TEXT("\n<%s>static_assert</>(<%s>false</>, \"You will need to add DOREPLIFETIME_WITH_PARAMS(%s, %s, %s) to GetLifetimeReplicatedProps\");"),
 				*HeaderViewSyntaxDecorators::KeywordDecorator,
 				*HeaderViewSyntaxDecorators::KeywordDecorator,
