BPHeaderView: Add context menu with generic options (select all, copy text)

diff --git a/Engine/Plugins/Editor/BlueprintHeaderView/Source/BlueprintHeaderView/BlueprintHeaderView.Build.cs b/Engine/Plugins/Editor/BlueprintHeaderView/Source/BlueprintHeaderView/BlueprintHeaderView.Build.cs
--- a/Engine/Plugins/Editor/BlueprintHeaderView/Source/BlueprintHeaderView/BlueprintHeaderView.Build.cs
+++ b/Engine/Plugins/Editor/BlueprintHeaderView/Source/BlueprintHeaderView/BlueprintHeaderView.Build.cs
@@ -27,6 +27,7 @@ public class BlueprintHeaderView : ModuleRules
 				"EditorStyle",
 				"UnrealEd",
 				"BlueprintGraph",
+				"ApplicationCore",
 			}
 			);
 		
diff --git a/Engine/Plugins/Editor/BlueprintHeaderView/Source/BlueprintHeaderView/Private/SBlueprintHeaderView.cpp b/Engine/Plugins/Editor/BlueprintHeaderView/Source/BlueprintHeaderView/Private/SBlueprintHeaderView.cpp
--- a/Engine/Plugins/Editor/BlueprintHeaderView/Source/BlueprintHeaderView/Private/SBlueprintHeaderView.cpp
+++ b/Engine/Plugins/Editor/BlueprintHeaderView/Source/BlueprintHeaderView/Private/SBlueprintHeaderView.cpp
@@ -10,6 +10,8 @@
 #include "IContentBrowserSingleton.h"
 #include "Engine/Blueprint.h"
 #include "Framework/Text/SlateTextRun.h"
+#include "Framework/Commands/GenericCommands.h"
+#include "Framework/MultiBox/MultiBoxBuilder.h"
 #include "HeaderViewClassListItem.h"
 #include "HeaderViewFunctionListItem.h"
 #include "HeaderViewVariableListItem.h"
@@ -17,6 +19,7 @@
 #include "K2Node_FunctionEntry.h"
 #include "EditorStyleSet.h"
 #include "String/LineEndings.h"
+#include "HAL/PlatformApplicationMisc.h"
 
 #define LOCTEXT_NAMESPACE "SBlueprintHeaderView"
 
@@ -171,6 +174,15 @@ void SBlueprintHeaderView::Construct(const FArguments& InArgs)
 	const float PaddingAmount = 8.0f;
 	SelectedBlueprint = nullptr;
 
+	CommandList = MakeShared<FUICommandList>();
+	
+	CommandList->MapAction(FGenericCommands::Get().Copy,
+		FExecuteAction::CreateRaw(this, &SBlueprintHeaderView::OnCopy),
+		FCanExecuteAction::CreateRaw(this, &SBlueprintHeaderView::CanCopy));
+
+	CommandList->MapAction(FGenericCommands::Get().SelectAll,
+		FExecuteAction::CreateRaw(this, &SBlueprintHeaderView::OnSelectAll));
+
 	ChildSlot
 	[
 		SNew(SVerticalBox)
@@ -222,11 +234,22 @@ void SBlueprintHeaderView::Construct(const FArguments& InArgs)
 				SAssignNew(ListView, SListView<FHeaderViewListItemPtr>)
 				.ListItemsSource(&ListItems)
 				.OnGenerateRow(this, &SBlueprintHeaderView::GenerateRowForItem)
+				.OnContextMenuOpening(this, &SBlueprintHeaderView::OnContextMenuOpening)
 //			]
 		]
 	];
 }
 
+FReply SBlueprintHeaderView::OnKeyDown(const FGeometry& MyGeometry, const FKeyEvent& InKeyEvent)
+{
+	if (CommandList.IsValid() && CommandList->ProcessCommandBindings(InKeyEvent))
+	{
+		return FReply::Handled();
+	}
+
+	return FReply::Unhandled();
+}
+
 FText SBlueprintHeaderView::GetClassPickerText() const
 {
 	if (const UBlueprint* Blueprint = SelectedBlueprint.Get())
@@ -381,4 +404,40 @@ void SBlueprintHeaderView::PopulateVariableItems(const UBlueprint* Blueprint)
 	}
 }
 
+TSharedPtr<SWidget> SBlueprintHeaderView::OnContextMenuOpening()
+{
+	FMenuBuilder MenuBuilder(true, CommandList);
+
+	MenuBuilder.AddMenuEntry(FGenericCommands::Get().SelectAll);
+	MenuBuilder.AddMenuEntry(FGenericCommands::Get().Copy);
+
+	return MenuBuilder.MakeWidget();
+}
+
+void SBlueprintHeaderView::OnCopy() const
+{
+	const int32 StringReserveSize = 2048;
+	FString CopyText;
+	CopyText.Reserve(StringReserveSize);
+	for (const FHeaderViewListItemPtr& Item : ListItems)
+	{
+		if (ListView->IsItemSelected(Item))
+		{
+			CopyText += Item->GetRawItemString() + LINE_TERMINATOR;
+		}
+	}
+
+	FPlatformApplicationMisc::ClipboardCopy(*CopyText);
+}
+
+bool SBlueprintHeaderView::CanCopy() const
+{
+	return ListView->GetNumItemsSelected() > 0;
+}
+
+void SBlueprintHeaderView::OnSelectAll()
+{
+	ListView->SetItemSelection(ListItems, true);
+}
+
 #undef LOCTEXT_NAMESPACE
diff --git a/Engine/Plugins/Editor/BlueprintHeaderView/Source/BlueprintHeaderView/Private/SBlueprintHeaderView.h b/Engine/Plugins/Editor/BlueprintHeaderView/Source/BlueprintHeaderView/Private/SBlueprintHeaderView.h
--- a/Engine/Plugins/Editor/BlueprintHeaderView/Source/BlueprintHeaderView/Private/SBlueprintHeaderView.h
+++ b/Engine/Plugins/Editor/BlueprintHeaderView/Source/BlueprintHeaderView/Private/SBlueprintHeaderView.h
@@ -12,6 +12,7 @@ class SComboButton;
 class UBlueprint;
 class ITableRow;
 class STableViewBase;
+class FUICommandList;
 
 /** rich text decorators for BlueprintHeaderView Syntax Highlighting */
 namespace HeaderViewSyntaxDecorators
@@ -32,6 +33,9 @@ struct FHeaderViewListItem : public TSharedFromThis<FHeaderViewListItem>
 	/** Creates a basic list item containing some text */
 	static TSharedPtr<FHeaderViewListItem> Create(FString InRawString, FString InRichText);
 
+	/** Returns the raw item text for copy actions */
+	const FString& GetRawItemString() const { return RawItemString; }
+
 protected:
 	/** Empty base constructor hidden from public */
 	FHeaderViewListItem() {};
@@ -74,6 +78,10 @@ public:
 
 	void Construct(const FArguments& InArgs);
 
+	//~ SWidget interface
+	virtual FReply OnKeyDown(const FGeometry& MyGeometry, const FKeyEvent& InKeyEvent) override;
+	//~ End SWidget interface
+
 private:
 	/** Gets the text for the class picker combo button */
 	FText GetClassPickerText() const;
@@ -95,7 +103,19 @@ private:
 
 	/** Adds items to the list view representing all variables present in the given blueprints */
 	void PopulateVariableItems(const UBlueprint* Blueprint);
+
+	/** Creates a context menu for the list view */
+	TSharedPtr<SWidget> OnContextMenuOpening();
+
+	/** UI Command Functions */
+	void OnCopy() const;
+	bool CanCopy() const;
+	void OnSelectAll();
+
 private:
+	/** List of UI Commands for this scope */
+	TSharedPtr<FUICommandList> CommandList;
+
 	/** The blueprint currently being displayed by the header view */
 	TWeakObjectPtr<UBlueprint> SelectedBlueprint;
 
