BPHeaderView: Add support for Delegate properties

diff --git a/Engine/Plugins/Editor/BlueprintHeaderView/Source/BlueprintHeaderView/Private/HeaderViewVariableListItem.cpp b/Engine/Plugins/Editor/BlueprintHeaderView/Source/BlueprintHeaderView/Private/HeaderViewVariableListItem.cpp
--- a/Engine/Plugins/Editor/BlueprintHeaderView/Source/BlueprintHeaderView/Private/HeaderViewVariableListItem.cpp
+++ b/Engine/Plugins/Editor/BlueprintHeaderView/Source/BlueprintHeaderView/Private/HeaderViewVariableListItem.cpp
@@ -51,6 +51,13 @@ FHeaderViewVariableListItem::FHeaderViewVariableListItem(const FBPVariableDescri
 		}
 	}
 
+	// Add Delegate type declaration if needed
+	// i.e. DECLARE_DYNAMIC_MULTICAST_DELEGATE_OneParam(FDelegateTypeName, ParamType, ParamName);
+	if (const FMulticastDelegateProperty* AsDelegate = CastField<FMulticastDelegateProperty>(&VarProperty))
+	{
+		FormatDelegateDeclaration(*AsDelegate);
+	}
+
 	// Add Deprecation message if present
 	if (VarProperty.HasAnyPropertyFlags(CPF_Deprecated))
 	{
@@ -101,13 +108,70 @@ FHeaderViewVariableListItem::FHeaderViewVariableListItem(const FBPVariableDescri
 	UE::String::ToHostLineEndingsInline(RichTextString);
 }
 
+void FHeaderViewVariableListItem::FormatDelegateDeclaration(const FMulticastDelegateProperty& DelegateProp)
+{
+	if (const UFunction* SigFunction = DelegateProp.SignatureFunction)
+	{
+		static const TMap<int32, FString> ParamStringsMap =
+		{
+			{-1, TEXT("_TooManyParams")},
+			{0, TEXT("")},
+			{1, TEXT("_OneParam")},
+			{2, TEXT("_TwoParams")},
+			{3, TEXT("_ThreeParams")},
+			{4, TEXT("_FourParams")},
+			{5, TEXT("_FiveParams")},
+			{6, TEXT("_SixParams")},
+			{7, TEXT("_SevenParams")},
+			{8, TEXT("_EightParams")},
+			{9, TEXT("_NineParams")}
+		};
+
+		const FString* NumParamsString = ParamStringsMap.Find(SigFunction->NumParms);
+		if (!NumParamsString)
+		{
+			NumParamsString = &ParamStringsMap[-1];
+		}
+
+		TArray<FString> MacroParams;
+		MacroParams.Emplace(GetCPPTypenameForProperty(&DelegateProp));
+
+		for (TFieldIterator<FProperty> ParamIt(SigFunction); ParamIt; ++ParamIt)
+		{
+			MacroParams.Emplace(GetCPPTypenameForProperty(*ParamIt));
+			if (ParamIt->HasAnyPropertyFlags(CPF_ReferenceParm))
+			{
+				MacroParams.Last().Append(TEXT("&"));
+			}
+
+			MacroParams.Emplace(ParamIt->GetAuthoredName());
+		}
+
+		RawItemString += FString::Printf(TEXT("\nDECLARE_DYNAMIC_MULTICAST_DELEGATE%s(%s);"), **NumParamsString, *FString::Join(MacroParams, TEXT(", ")));
+
+		// Add rich text decorators for macro parameters
+		for (int32 ParmIdx = 0; ParmIdx < MacroParams.Num(); ++ParmIdx)
+		{
+			// the delegate name (idx 0) and parameter types (odd number idx) are typenames, all other parameters are identifiers
+			const bool bIsTypename = (ParmIdx == 0 || ParmIdx % 2 != 0);
+			MacroParams[ParmIdx] = FString::Printf(TEXT("<%s>%s</>"), bIsTypename ? *HeaderViewSyntaxDecorators::TypenameDecorator : *HeaderViewSyntaxDecorators::IdentifierDecorator, *MacroParams[ParmIdx]);
+		}
+
+		RichTextString += FString::Printf(TEXT("\n<%s>DECLARE_DYNAMIC_MULTICAST_DELEGATE%s</>(%s);"), *HeaderViewSyntaxDecorators::MacroDecorator, **NumParamsString, *FString::Join(MacroParams, TEXT(", ")));
+	}
+}
+
 FString FHeaderViewVariableListItem::GetConditionalUPropertySpecifiers(const FProperty& VarProperty) const
 {
 	TArray<FString> PropertySpecifiers;
 
 	if (!VarProperty.HasMetaData(FBlueprintMetadata::MD_Private) || !VarProperty.GetBoolMetaData(FBlueprintMetadata::MD_Private))
 	{
-		if (VarProperty.HasAllPropertyFlags(CPF_BlueprintReadOnly))
+		if (VarProperty.HasAnyPropertyFlags(CPF_BlueprintAssignable))
+		{
+			PropertySpecifiers.Emplace(TEXT("BlueprintAssignable"));
+		}
+		else if (VarProperty.HasAnyPropertyFlags(CPF_BlueprintReadOnly))
 		{
 			PropertySpecifiers.Emplace(TEXT("BlueprintReadOnly"));
 		}
diff --git a/Engine/Plugins/Editor/BlueprintHeaderView/Source/BlueprintHeaderView/Private/HeaderViewVariableListItem.h b/Engine/Plugins/Editor/BlueprintHeaderView/Source/BlueprintHeaderView/Private/HeaderViewVariableListItem.h
--- a/Engine/Plugins/Editor/BlueprintHeaderView/Source/BlueprintHeaderView/Private/HeaderViewVariableListItem.h
+++ b/Engine/Plugins/Editor/BlueprintHeaderView/Source/BlueprintHeaderView/Private/HeaderViewVariableListItem.h
@@ -16,6 +16,9 @@ struct FHeaderViewVariableListItem : public FHeaderViewListItem
 protected:
 	FHeaderViewVariableListItem(const FBPVariableDescription& VariableDesc, const FProperty& VarProperty);
 
+	/** Formats a line declaring a delegate type and appends it to the item strings */
+	void FormatDelegateDeclaration(const FMulticastDelegateProperty& DelegateProp);
+
 	/** Returns a string containing the specifiers for the UPROPERTY line */
 	FString GetConditionalUPropertySpecifiers(const FProperty& VarProperty) const;
 
