BPHeaderView: Add support for blueprint added components

diff --git a/Engine/Plugins/Editor/BlueprintHeaderView/Source/BlueprintHeaderView/Private/HeaderViewVariableListItem.cpp b/Engine/Plugins/Editor/BlueprintHeaderView/Source/BlueprintHeaderView/Private/HeaderViewVariableListItem.cpp
--- a/Engine/Plugins/Editor/BlueprintHeaderView/Source/BlueprintHeaderView/Private/HeaderViewVariableListItem.cpp
+++ b/Engine/Plugins/Editor/BlueprintHeaderView/Source/BlueprintHeaderView/Private/HeaderViewVariableListItem.cpp
@@ -6,12 +6,12 @@
 #include "Engine/Blueprint.h"
 #include "String/LineEndings.h"
 
-FHeaderViewListItemPtr FHeaderViewVariableListItem::Create(const FBPVariableDescription& VariableDesc, const FProperty& VarProperty)
+FHeaderViewListItemPtr FHeaderViewVariableListItem::Create(const FBPVariableDescription* VariableDesc, const FProperty& VarProperty)
 {
 	return MakeShareable(new FHeaderViewVariableListItem(VariableDesc, VarProperty));
 }
 
-FHeaderViewVariableListItem::FHeaderViewVariableListItem(const FBPVariableDescription& VariableDesc, const FProperty& VarProperty)
+FHeaderViewVariableListItem::FHeaderViewVariableListItem(const FBPVariableDescription* VariableDesc, const FProperty& VarProperty)
 {
 	// Format comment
 	{
@@ -29,7 +29,7 @@ FHeaderViewVariableListItem::FHeaderViewVariableListItem(const FBPVariableDescri
 	if (VarProperty.HasAnyPropertyFlags(CPF_Net))
 	{
 		FString ClassName = GetOwningClassName(VarProperty);
-		if (VariableDesc.ReplicationCondition == COND_None)
+		if (!VariableDesc || VariableDesc->ReplicationCondition == COND_None)
 		{
 			RawItemString += FString::Printf(TEXT("\nstatic_assert(false, \"You will need to add DOREPLIFETIME(%s, %s) to GetLifetimeReplicatedProps\");"), *ClassName, *VarProperty.GetAuthoredName());
 			RichTextString += FString::Printf(TEXT("\n<%s>static_assert</>(<%s>false</>, \"You will need to add DOREPLIFETIME(%s, %s) to GetLifetimeReplicatedProps\");"), 
@@ -40,7 +40,7 @@ FHeaderViewVariableListItem::FHeaderViewVariableListItem(const FBPVariableDescri
 		}
 		else
 		{
-			FString ConditionString = StaticEnum<ELifetimeCondition>()->GetAuthoredNameStringByValue(VariableDesc.ReplicationCondition);
+			FString ConditionString = StaticEnum<ELifetimeCondition>()->GetAuthoredNameStringByValue(VariableDesc->ReplicationCondition);
 			RawItemString += FString::Printf(TEXT("\nstatic_assert(false, \"You will need to add DOREPLIFETIME_WITH_PARAMS(%s, %s, %s) to GetLifetimeReplicatedProps\");"), *ClassName, *VarProperty.GetAuthoredName(), *ConditionString);
 			RichTextString += FString::Printf(TEXT("\n<%s>static_assert</>(<%s>false</>, \"You will need to add DOREPLIFETIME_WITH_PARAMS(%s, %s, %s) to GetLifetimeReplicatedProps\");"),
 				*HeaderViewSyntaxDecorators::KeywordDecorator,
@@ -86,7 +86,7 @@ FHeaderViewVariableListItem::FHeaderViewVariableListItem(const FBPVariableDescri
 	// Add the variable declaration line
 	// i.e. Type VariableName;
 	{
-		const FString Typename = GetCPPTypenameForProperty(&VarProperty);
+		const FString Typename = GetCPPTypenameForProperty(&VarProperty, /*bIsMemberProperty=*/true);
 		const FString VarName = VarProperty.GetAuthoredName();
 
 		RawItemString += FString::Printf(TEXT("\n%s %s;"), *Typename, *VarName);
@@ -175,13 +175,21 @@ FString FHeaderViewVariableListItem::GetConditionalUPropertySpecifiers(const FPr
 		{
 			PropertySpecifiers.Emplace(TEXT("BlueprintReadOnly"));
 		}
-		else
+		else if (VarProperty.HasAnyPropertyFlags(CPF_Edit))
 		{
 			PropertySpecifiers.Emplace(TEXT("BlueprintReadWrite"));
 		}
+		else
+		{
+			PropertySpecifiers.Emplace(TEXT("BlueprintVisible"));
+		}
 	}
 
-	if (!VarProperty.HasAnyPropertyFlags(CPF_DisableEditOnInstance | CPF_DisableEditOnTemplate))
+	if (!VarProperty.HasAnyPropertyFlags(CPF_Edit))
+	{
+		PropertySpecifiers.Emplace(TEXT("VisibleAnywhere"));
+	}
+	else if (!VarProperty.HasAnyPropertyFlags(CPF_DisableEditOnInstance | CPF_DisableEditOnTemplate))
 	{
 		PropertySpecifiers.Emplace(TEXT("EditAnywhere"));
 	}
@@ -243,41 +251,44 @@ FString FHeaderViewVariableListItem::GetConditionalUPropertySpecifiers(const FPr
 
 	// meta specifiers
 	{
-		// These metadata keys are handled explicitly elsewhere
-		static const TArray<FName> IgnoreMetaData = 
+		if (const TMap<FName, FString>* MetaDataMap = VarProperty.GetMetaDataMap())
 		{
-			FBlueprintMetadata::MD_DeprecationMessage,
-			FBlueprintMetadata::MD_FunctionCategory,
-			FBlueprintMetadata::MD_Tooltip,
-			FBlueprintMetadata::MD_Private,
-			FEdMode::MD_MakeEditWidget
-		};
-
-		TArray<FString> MetaSpecifiers;
-
-		for (const TPair<FName, FString>& MetaPair : *VarProperty.GetMetaDataMap())
-		{
-			// Don't take this metadata if we handled it explicitly
-			if (IgnoreMetaData.Contains(MetaPair.Key))
+			// These metadata keys are handled explicitly elsewhere
+			static const TArray<FName> IgnoreMetaData = 
 			{
-				continue;
-			}
+				FBlueprintMetadata::MD_DeprecationMessage,
+				FBlueprintMetadata::MD_FunctionCategory,
+				FBlueprintMetadata::MD_Tooltip,
+				FBlueprintMetadata::MD_Private,
+				FEdMode::MD_MakeEditWidget
+			};
 
-			// Don't add the DisplayName metadata key if the display name is just the friendly name generated by the editor anyway
-			if (MetaPair.Key == FBlueprintMetadata::MD_DisplayName)
+			TArray<FString> MetaSpecifiers;
+
+			for (const TPair<FName, FString>& MetaPair : *MetaDataMap)
 			{
-				if (MetaPair.Value == FName::NameToDisplayString(VarProperty.GetName(), !!CastField<FBoolProperty>(&VarProperty)))
+				// Don't take this metadata if we handled it explicitly
+				if (IgnoreMetaData.Contains(MetaPair.Key))
 				{
 					continue;
 				}
-			}
 
-			MetaSpecifiers.Emplace(FString::Printf(TEXT("%s=\"%s\""), *MetaPair.Key.ToString(), *MetaPair.Value));
-		}
+				// Don't add the DisplayName metadata key if the display name is just the friendly name generated by the editor anyway
+				if (MetaPair.Key == FBlueprintMetadata::MD_DisplayName)
+				{
+					if (MetaPair.Value == FName::NameToDisplayString(VarProperty.GetName(), !!CastField<FBoolProperty>(&VarProperty)))
+					{
+						continue;
+					}
+				}
 
-		if (!MetaSpecifiers.IsEmpty())
-		{
-			PropertySpecifiers.Emplace(FString::Printf(TEXT("meta=(%s)"), *FString::Join(MetaSpecifiers, TEXT(", "))));
+				MetaSpecifiers.Emplace(FString::Printf(TEXT("%s=\"%s\""), *MetaPair.Key.ToString(), *MetaPair.Value));
+			}
+
+			if (!MetaSpecifiers.IsEmpty())
+			{
+				PropertySpecifiers.Emplace(FString::Printf(TEXT("meta=(%s)"), *FString::Join(MetaSpecifiers, TEXT(", "))));
+			}
 		}
 	}
 
diff --git a/Engine/Plugins/Editor/BlueprintHeaderView/Source/BlueprintHeaderView/Private/HeaderViewVariableListItem.h b/Engine/Plugins/Editor/BlueprintHeaderView/Source/BlueprintHeaderView/Private/HeaderViewVariableListItem.h
--- a/Engine/Plugins/Editor/BlueprintHeaderView/Source/BlueprintHeaderView/Private/HeaderViewVariableListItem.h
+++ b/Engine/Plugins/Editor/BlueprintHeaderView/Source/BlueprintHeaderView/Private/HeaderViewVariableListItem.h
@@ -11,10 +11,10 @@ struct FBPVariableDescription;
 struct FHeaderViewVariableListItem : public FHeaderViewListItem
 {
 	/** Creates a list item for the Header view representing a variable declaration for the given blueprint variable */
-	static FHeaderViewListItemPtr Create(const FBPVariableDescription& VariableDesc, const FProperty& VarProperty);
+	static FHeaderViewListItemPtr Create(const FBPVariableDescription* VariableDesc, const FProperty& VarProperty);
 
 protected:
-	FHeaderViewVariableListItem(const FBPVariableDescription& VariableDesc, const FProperty& VarProperty);
+	FHeaderViewVariableListItem(const FBPVariableDescription* VariableDesc, const FProperty& VarProperty);
 
 	/** Formats a line declaring a delegate type and appends it to the item strings */
 	void FormatDelegateDeclaration(const FMulticastDelegateProperty& DelegateProp);
diff --git a/Engine/Plugins/Editor/BlueprintHeaderView/Source/BlueprintHeaderView/Private/SBlueprintHeaderView.cpp b/Engine/Plugins/Editor/BlueprintHeaderView/Source/BlueprintHeaderView/Private/SBlueprintHeaderView.cpp
--- a/Engine/Plugins/Editor/BlueprintHeaderView/Source/BlueprintHeaderView/Private/SBlueprintHeaderView.cpp
+++ b/Engine/Plugins/Editor/BlueprintHeaderView/Source/BlueprintHeaderView/Private/SBlueprintHeaderView.cpp
@@ -154,12 +154,24 @@ void FHeaderViewListItem::FormatCommentString(FString InComment, FString& OutRaw
 	OutRichString = FString::Printf(TEXT("<%s>%s</>"), *HeaderViewSyntaxDecorators::CommentDecorator, *InComment);
 }
 
-FString FHeaderViewListItem::GetCPPTypenameForProperty(const FProperty* InProperty)
+FString FHeaderViewListItem::GetCPPTypenameForProperty(const FProperty* InProperty, bool bIsMemberProperty/*=false*/)
 {
 	if (InProperty)
 	{
 		FString ExtendedTypeText;
-		return InProperty->GetCPPType(&ExtendedTypeText) + ExtendedTypeText;
+		FString Typename = InProperty->GetCPPType(&ExtendedTypeText) + ExtendedTypeText;
+		
+		if (bIsMemberProperty)
+		{
+			if (const FObjectProperty* ObjectProperty = CastField<FObjectProperty>(InProperty))
+			{
+				// Replace native pointer with TObjectPtr
+				Typename.LeftChopInline(1);
+				Typename = FString::Printf(TEXT("TObjectPtr<%s>"), *Typename);
+			}
+		}
+
+		return Typename;
 	}
 	else
 	{
@@ -373,32 +385,40 @@ void SBlueprintHeaderView::PopulateVariableItems(const UBlueprint* Blueprint)
 
 		// We should only add an access specifier line if the previous variable was a different one
 		int32 PrevAccessSpecifier = 0;
-		for (const FBPVariableDescription& VariableDesc : Blueprint->NewVariables)
+		for (TFieldIterator<FProperty> PropertyIt(Blueprint->SkeletonGeneratedClass, EFieldIteratorFlags::ExcludeSuper); PropertyIt; ++PropertyIt)
 		{
-			if (const FProperty* VarProperty = Blueprint->SkeletonGeneratedClass->FindPropertyByName(VariableDesc.VarName))
+			if (const FProperty* VarProperty = *PropertyIt)
 			{
-				const int32 AccessSpecifier = VarProperty->GetBoolMetaData(FBlueprintMetadata::MD_Private) ? Private : Public;
-				if (AccessSpecifier != PrevAccessSpecifier)
+				if (VarProperty->HasAnyPropertyFlags(CPF_BlueprintVisible))
 				{
-					switch (AccessSpecifier)
+					const int32 AccessSpecifier = VarProperty->GetBoolMetaData(FBlueprintMetadata::MD_Private) ? Private : Public;
+					if (AccessSpecifier != PrevAccessSpecifier)
+					{
+						switch (AccessSpecifier)
+						{
+						case Public:
+							ListItems.Add(FHeaderViewListItem::Create(TEXT("public:"), FString::Printf(TEXT("<%s>public</>:"), *HeaderViewSyntaxDecorators::KeywordDecorator)));
+							break;
+						case Private:
+							ListItems.Add(FHeaderViewListItem::Create(TEXT("private:"), FString::Printf(TEXT("<%s>private</>:"), *HeaderViewSyntaxDecorators::KeywordDecorator)));
+							break;
+						}
+
+						PrevAccessSpecifier = AccessSpecifier;
+					}
+					else
 					{
-					case Public:
-						ListItems.Add(FHeaderViewListItem::Create(TEXT("public:"), FString::Printf(TEXT("<%s>public</>:"), *HeaderViewSyntaxDecorators::KeywordDecorator)));
-						break;
-					case Private:
-						ListItems.Add(FHeaderViewListItem::Create(TEXT("private:"), FString::Printf(TEXT("<%s>private</>:"), *HeaderViewSyntaxDecorators::KeywordDecorator)));
-						break;
+						// add an empty line to space variables out
+						ListItems.Add(FHeaderViewListItem::Create(TEXT(""), TEXT("")));
 					}
 
-					PrevAccessSpecifier = AccessSpecifier;
-				}
-				else
-				{
-					// add an empty line to space variables out
-					ListItems.Add(FHeaderViewListItem::Create(TEXT(""), TEXT("")));
-				}
+					const FBPVariableDescription* VariableDesc = Blueprint->NewVariables.FindByPredicate([&VarProperty](const FBPVariableDescription& Desc)
+						{
+							return Desc.VarName == VarProperty->GetFName();
+						});
 
-				ListItems.Add(FHeaderViewVariableListItem::Create(VariableDesc, *VarProperty));
+					ListItems.Add(FHeaderViewVariableListItem::Create(VariableDesc, *VarProperty));
+				}
 			}
 		}
 	}
diff --git a/Engine/Plugins/Editor/BlueprintHeaderView/Source/BlueprintHeaderView/Private/SBlueprintHeaderView.h b/Engine/Plugins/Editor/BlueprintHeaderView/Source/BlueprintHeaderView/Private/SBlueprintHeaderView.h
--- a/Engine/Plugins/Editor/BlueprintHeaderView/Source/BlueprintHeaderView/Private/SBlueprintHeaderView.h
+++ b/Engine/Plugins/Editor/BlueprintHeaderView/Source/BlueprintHeaderView/Private/SBlueprintHeaderView.h
@@ -54,7 +54,7 @@ protected:
 	 * returns a string representing the full C++ typename for the given property, 
 	 * including template params for container types
 	 */
-	static FString GetCPPTypenameForProperty(const FProperty* InProperty);
+	static FString GetCPPTypenameForProperty(const FProperty* InProperty, bool bIsMemberProperty = false);
 
 protected:
 	/** A rich text representation of the item, including syntax highlighting and errors */
