BP Header View: Add Config options for syntax highlighting and field sorting

diff --git a/Engine/Config/BaseEditorPerProjectUserSettings.ini b/Engine/Config/BaseEditorPerProjectUserSettings.ini
--- a/Engine/Config/BaseEditorPerProjectUserSettings.ini
+++ b/Engine/Config/BaseEditorPerProjectUserSettings.ini
@@ -928,3 +928,7 @@ Startup=AutomaticButHidden
 bEnabled=True
 bEnableReinstancing=True
 bAutomaticallyCompileNewClasses=True
+
+[/Script/BlueprintHeaderView.BlueprintHeaderViewSettings]
+SyntaxColors=(Comment=(R=0.300000,G=0.700000,B=0.100000,A=1.000000),Macro=(R=0.600000,G=0.200000,B=0.800000,A=1.000000),Typename=(R=0.000000,G=0.300000,B=0.300000,A=1.000000),Identifier=(R=1.000000,G=1.000000,B=1.000000,A=1.000000),Keyword=(R=0.000000,G=0.400000,B=0.800000,A=1.000000))
+SortMethod=None
diff --git a/Engine/Plugins/Editor/BlueprintHeaderView/Source/BlueprintHeaderView/BlueprintHeaderView.Build.cs b/Engine/Plugins/Editor/BlueprintHeaderView/Source/BlueprintHeaderView/BlueprintHeaderView.Build.cs
--- a/Engine/Plugins/Editor/BlueprintHeaderView/Source/BlueprintHeaderView/BlueprintHeaderView.Build.cs
+++ b/Engine/Plugins/Editor/BlueprintHeaderView/Source/BlueprintHeaderView/BlueprintHeaderView.Build.cs
@@ -28,6 +28,7 @@ public class BlueprintHeaderView : ModuleRules
 				"UnrealEd",
 				"BlueprintGraph",
 				"ApplicationCore",
+				"DeveloperSettings",
 			}
 			);
 		
diff --git a/Engine/Plugins/Editor/BlueprintHeaderView/Source/BlueprintHeaderView/Private/BlueprintHeaderViewSettings.cpp b/Engine/Plugins/Editor/BlueprintHeaderView/Source/BlueprintHeaderView/Private/BlueprintHeaderViewSettings.cpp
new file mode 100644
--- /dev/null
+++ b/Engine/Plugins/Editor/BlueprintHeaderView/Source/BlueprintHeaderView/Private/BlueprintHeaderViewSettings.cpp
@@ -0,0 +1,23 @@
+// Copyright Epic Games, Inc. All Rights Reserved.
+
+#include "BlueprintHeaderViewSettings.h"
+
+UBlueprintHeaderViewSettings::UBlueprintHeaderViewSettings()
+{
+	
+}
+
+FName UBlueprintHeaderViewSettings::GetCategoryName() const
+{
+	return TEXT("Plugins");
+}
+
+FText UBlueprintHeaderViewSettings::GetSectionText() const
+{
+	return NSLOCTEXT("BlueprintHeaderViewSettings", "HeaderViewSectionText", "Blueprint Header View");
+}
+
+FName UBlueprintHeaderViewSettings::GetSectionName() const
+{
+	return TEXT("Blueprint Header View");
+}
diff --git a/Engine/Plugins/Editor/BlueprintHeaderView/Source/BlueprintHeaderView/Private/SBlueprintHeaderView.cpp b/Engine/Plugins/Editor/BlueprintHeaderView/Source/BlueprintHeaderView/Private/SBlueprintHeaderView.cpp
--- a/Engine/Plugins/Editor/BlueprintHeaderView/Source/BlueprintHeaderView/Private/SBlueprintHeaderView.cpp
+++ b/Engine/Plugins/Editor/BlueprintHeaderView/Source/BlueprintHeaderView/Private/SBlueprintHeaderView.cpp
@@ -6,6 +6,7 @@
 //#include "Widgets/Layout/SScrollBox.h"
 #include "Widgets/Input/SComboButton.h"
 #include "Widgets/Text/SRichTextBlock.h"
+#include "PropertyEditorModule.h"
 #include "ContentBrowserModule.h"
 #include "IContentBrowserSingleton.h"
 #include "Engine/Blueprint.h"
@@ -20,6 +21,7 @@
 #include "EditorStyleSet.h"
 #include "String/LineEndings.h"
 #include "HAL/PlatformApplicationMisc.h"
+#include "BlueprintHeaderViewSettings.h"
 
 #define LOCTEXT_NAMESPACE "SBlueprintHeaderView"
 
@@ -90,15 +92,8 @@ namespace
 
 TSharedRef<SWidget> FHeaderViewListItem::GenerateWidgetForItem()
 {
-	// TODO: colors not final, move this struct into config
-	static const struct
-	{
-		FSlateColor Comment = FLinearColor(0.3f, 0.7f, 0.1f, 1.0f);
-		FSlateColor Macro = FLinearColor(0.6f, 0.2f, 0.8f, 1.0f);
-		FSlateColor Typename = FLinearColor::White;
-		FSlateColor Identifier = FLinearColor::White;
-		FSlateColor Keyword = FLinearColor(0.0f, 0.4f, 0.8f, 1.0f);
-	} SyntaxColors;
+	const UBlueprintHeaderViewSettings* HeaderViewSettings = GetDefault<UBlueprintHeaderViewSettings>();
+	const FHeaderViewSyntaxColors& SyntaxColors = HeaderViewSettings->SyntaxColors;
 
 	return SNew(SBox)
 		.HAlign(HAlign_Fill)
@@ -234,6 +229,37 @@ void SBlueprintHeaderView::Construct(const FArguments& InArgs)
 					]
 				]
 			]
+			+SHorizontalBox::Slot()
+			.HAlign(HAlign_Right)
+			[
+				SNew(SComboButton)
+				.HasDownArrow(false)
+				.ForegroundColor(FSlateColor::UseForeground())
+				.ButtonStyle(FAppStyle::Get(), "SimpleButton")
+				.OnGetMenuContent(this, &SBlueprintHeaderView::GetSettingsMenuContent)
+				.ButtonContent()
+				[
+					SNew(SHorizontalBox)
+					+SHorizontalBox::Slot()
+					.AutoWidth()
+					.HAlign(HAlign_Center)
+					[
+						SNew(SImage)
+						.Image(FAppStyle::Get().GetBrush("Icons.Toolbar.Settings"))
+						.ColorAndOpacity(FSlateColor::UseForeground())
+					]
+					+SHorizontalBox::Slot()
+					.Padding(FMargin(5, 0, 0, 0))
+					.VAlign(VAlign_Center)
+					.AutoWidth()
+					[
+						SNew(STextBlock)
+						.TextStyle(FAppStyle::Get(), "NormalText")
+						.Text(LOCTEXT("SettingsLabel", "Settings"))
+						.ColorAndOpacity(FSlateColor::UseForeground())
+					]
+				]
+			]
 		]
 		+SVerticalBox::Slot()
 		.Padding(FMargin(PaddingAmount))
@@ -262,6 +288,13 @@ FReply SBlueprintHeaderView::OnKeyDown(const FGeometry& MyGeometry, const FKeyEv
 	return FReply::Unhandled();
 }
 
+void SBlueprintHeaderView::NotifyPostChange(const FPropertyChangedEvent& PropertyChangedEvent, FProperty* PropertyThatChanged)
+{
+	UBlueprintHeaderViewSettings* BlueprintHeaderViewSettings = GetMutableDefault<UBlueprintHeaderViewSettings>();
+	check(BlueprintHeaderViewSettings);
+	BlueprintHeaderViewSettings->SaveConfig();
+}
+
 FText SBlueprintHeaderView::GetClassPickerText() const
 {
 	if (const UBlueprint* Blueprint = SelectedBlueprint.Get())
@@ -290,6 +323,20 @@ TSharedRef<SWidget> SBlueprintHeaderView::GetClassPickerMenuContent()
 		];
 }
 
+TSharedRef<SWidget> SBlueprintHeaderView::GetSettingsMenuContent()
+{
+	FPropertyEditorModule& PropertyEditorModule = FModuleManager::LoadModuleChecked<FPropertyEditorModule>("PropertyEditor");	
+	FDetailsViewArgs DetailsViewArgs;
+	{
+		DetailsViewArgs.bShowPropertyMatrixButton = false;
+		DetailsViewArgs.bShowOptions = false;
+		DetailsViewArgs.NotifyHook = this;
+	}
+	TSharedRef<IDetailsView> DetailsView = PropertyEditorModule.CreateDetailView(DetailsViewArgs);
+	DetailsView->SetObject(GetMutableDefault<UBlueprintHeaderViewSettings>());
+	return DetailsView;
+}
+
 void SBlueprintHeaderView::OnAssetSelected(const FAssetData& SelectedAsset)
 {
 	ClassPickerComboButton->SetIsOpen(false);
@@ -332,9 +379,12 @@ void SBlueprintHeaderView::PopulateFunctionItems(const UBlueprint* Blueprint)
 {
 	if (Blueprint)
 	{
+		TArray<const UEdGraph*> FunctionGraphs;
+		GatherFunctionGraphs(Blueprint, FunctionGraphs);
+
 		// We should only add an access specifier line if the previous function was a different one
 		int32 PrevAccessSpecifier = 0;
-		for (const UEdGraph* FunctionGraph : Blueprint->FunctionGraphs)
+		for (const UEdGraph* FunctionGraph : FunctionGraphs)
 		{
 			if (FunctionGraph && !UEdGraphSchema_K2::IsConstructionScript(FunctionGraph))
 			{
@@ -344,12 +394,13 @@ void SBlueprintHeaderView::PopulateFunctionItems(const UBlueprint* Blueprint)
 				if (ensure(EntryNodes.Num() == 1 && EntryNodes[0]))
 				{
 					int32 AccessSpecifier = EntryNodes[0]->GetFunctionFlags() & FUNC_AccessSpecifiers;
+					AccessSpecifier = AccessSpecifier ? AccessSpecifier : FUNC_Public; //< Blueprint OnRep functions don't have one of these flags set, default to public in this case and others like it
 
 					if (AccessSpecifier != PrevAccessSpecifier)
 					{
 						switch (AccessSpecifier)
 						{
-						default: //< Blueprint OnRep functions don't have one of these flags set, default to public in this case and others like it
+						default: 
 						case FUNC_Public:
 							ListItems.Add(FHeaderViewListItem::Create(TEXT("public:"), FString::Printf(TEXT("<%s>public</>:"), *HeaderViewSyntaxDecorators::KeywordDecorator)));
 							break;
@@ -376,6 +427,35 @@ void SBlueprintHeaderView::PopulateFunctionItems(const UBlueprint* Blueprint)
 	}
 }
 
+void SBlueprintHeaderView::GatherFunctionGraphs(const UBlueprint* Blueprint, TArray<const UEdGraph*>& OutFunctionGraphs)
+{
+	for (const UEdGraph* FunctionGraph : Blueprint->FunctionGraphs)
+	{
+		OutFunctionGraphs.Add(FunctionGraph);
+	}
+
+	const UBlueprintHeaderViewSettings* BlueprintHeaderViewSettings = GetDefault<UBlueprintHeaderViewSettings>();
+	if (BlueprintHeaderViewSettings->SortMethod == EHeaderViewSortMethod::SortByAccessSpecifier)
+	{
+		OutFunctionGraphs.Sort([](const UEdGraph& LeftGraph, const UEdGraph& RightGraph)
+			{
+				TArray<UK2Node_FunctionEntry*> LeftEntryNodes;
+				LeftGraph.GetNodesOfClass<UK2Node_FunctionEntry>(LeftEntryNodes);
+				TArray<UK2Node_FunctionEntry*> RightEntryNodes;
+				RightGraph.GetNodesOfClass<UK2Node_FunctionEntry>(RightEntryNodes);
+				if (ensure(LeftEntryNodes.Num() == 1 && LeftEntryNodes[0] && RightEntryNodes.Num() == 1 && RightEntryNodes[0]))
+				{
+					const int32 LeftAccessSpecifier = LeftEntryNodes[0]->GetFunctionFlags() & FUNC_AccessSpecifiers;
+					const int32 RightAccessSpecifier = RightEntryNodes[0]->GetFunctionFlags() & FUNC_AccessSpecifiers;
+
+					return (LeftAccessSpecifier == FUNC_Public && RightAccessSpecifier != FUNC_Public) || (LeftAccessSpecifier == FUNC_Protected && RightAccessSpecifier == FUNC_Private);
+				}
+
+				return false;
+			});
+	}
+}
+
 void SBlueprintHeaderView::PopulateVariableItems(const UBlueprint* Blueprint)
 {
 	if (Blueprint)
@@ -383,45 +463,65 @@ void SBlueprintHeaderView::PopulateVariableItems(const UBlueprint* Blueprint)
 		const int32 Private = 2;
 		const int32 Public = 1;
 
+		TArray<const FProperty*> VarProperties;
+		GatherProperties(Blueprint, VarProperties);
+
 		// We should only add an access specifier line if the previous variable was a different one
 		int32 PrevAccessSpecifier = 0;
-		for (TFieldIterator<FProperty> PropertyIt(Blueprint->SkeletonGeneratedClass, EFieldIteratorFlags::ExcludeSuper); PropertyIt; ++PropertyIt)
+		for (const FProperty* VarProperty : VarProperties)
 		{
-			if (const FProperty* VarProperty = *PropertyIt)
+			const int32 AccessSpecifier = VarProperty->GetBoolMetaData(FBlueprintMetadata::MD_Private) ? Private : Public;
+			if (AccessSpecifier != PrevAccessSpecifier)
 			{
-				if (VarProperty->HasAnyPropertyFlags(CPF_BlueprintVisible))
+				switch (AccessSpecifier)
 				{
-					const int32 AccessSpecifier = VarProperty->GetBoolMetaData(FBlueprintMetadata::MD_Private) ? Private : Public;
-					if (AccessSpecifier != PrevAccessSpecifier)
-					{
-						switch (AccessSpecifier)
-						{
-						case Public:
-							ListItems.Add(FHeaderViewListItem::Create(TEXT("public:"), FString::Printf(TEXT("<%s>public</>:"), *HeaderViewSyntaxDecorators::KeywordDecorator)));
-							break;
-						case Private:
-							ListItems.Add(FHeaderViewListItem::Create(TEXT("private:"), FString::Printf(TEXT("<%s>private</>:"), *HeaderViewSyntaxDecorators::KeywordDecorator)));
-							break;
-						}
+				case Public:
+					ListItems.Add(FHeaderViewListItem::Create(TEXT("public:"), FString::Printf(TEXT("<%s>public</>:"), *HeaderViewSyntaxDecorators::KeywordDecorator)));
+					break;
+				case Private:
+					ListItems.Add(FHeaderViewListItem::Create(TEXT("private:"), FString::Printf(TEXT("<%s>private</>:"), *HeaderViewSyntaxDecorators::KeywordDecorator)));
+					break;
+				}
 
-						PrevAccessSpecifier = AccessSpecifier;
-					}
-					else
-					{
-						// add an empty line to space variables out
-						ListItems.Add(FHeaderViewListItem::Create(TEXT(""), TEXT("")));
-					}
+				PrevAccessSpecifier = AccessSpecifier;
+			}
+			else
+			{
+				// add an empty line to space variables out
+				ListItems.Add(FHeaderViewListItem::Create(TEXT(""), TEXT("")));
+			}
 
-					const FBPVariableDescription* VariableDesc = Blueprint->NewVariables.FindByPredicate([&VarProperty](const FBPVariableDescription& Desc)
-						{
-							return Desc.VarName == VarProperty->GetFName();
-						});
+			const FBPVariableDescription* VariableDesc = Blueprint->NewVariables.FindByPredicate([&VarProperty](const FBPVariableDescription& Desc)
+				{
+					return Desc.VarName == VarProperty->GetFName();
+				});
 
-					ListItems.Add(FHeaderViewVariableListItem::Create(VariableDesc, *VarProperty));
-				}
+			ListItems.Add(FHeaderViewVariableListItem::Create(VariableDesc, *VarProperty));
+		}
+	}
+}
+
+void SBlueprintHeaderView::GatherProperties(const UBlueprint* Blueprint, TArray<const FProperty*>& OutProperties)
+{
+	for (TFieldIterator<FProperty> PropertyIt(Blueprint->SkeletonGeneratedClass, EFieldIteratorFlags::ExcludeSuper); PropertyIt; ++PropertyIt)
+	{
+		if (const FProperty* VarProperty = *PropertyIt)
+		{
+			if (VarProperty->HasAnyPropertyFlags(CPF_BlueprintVisible))
+			{
+				OutProperties.Add(VarProperty);
 			}
 		}
 	}
+
+	const UBlueprintHeaderViewSettings* BlueprintHeaderViewSettings = GetDefault<UBlueprintHeaderViewSettings>();
+	if (BlueprintHeaderViewSettings->SortMethod == EHeaderViewSortMethod::SortByAccessSpecifier)
+	{
+		OutProperties.Sort([](const FProperty& LeftProp, const FProperty& RightProp)
+			{
+				return !LeftProp.GetBoolMetaData(FBlueprintMetadata::MD_Private) && RightProp.GetBoolMetaData(FBlueprintMetadata::MD_Private);
+			});
+	}
 }
 
 TSharedPtr<SWidget> SBlueprintHeaderView::OnContextMenuOpening()
diff --git a/Engine/Plugins/Editor/BlueprintHeaderView/Source/BlueprintHeaderView/Private/SBlueprintHeaderView.h b/Engine/Plugins/Editor/BlueprintHeaderView/Source/BlueprintHeaderView/Private/SBlueprintHeaderView.h
--- a/Engine/Plugins/Editor/BlueprintHeaderView/Source/BlueprintHeaderView/Private/SBlueprintHeaderView.h
+++ b/Engine/Plugins/Editor/BlueprintHeaderView/Source/BlueprintHeaderView/Private/SBlueprintHeaderView.h
@@ -7,12 +7,14 @@
 #include "Widgets/SWidget.h"
 #include "Widgets/SCompoundWidget.h"
 #include "Widgets/Views/SListView.h"
+#include "Misc/NotifyHook.h"
 
 class SComboButton;
 class UBlueprint;
 class ITableRow;
 class STableViewBase;
 class FUICommandList;
+class UEdGraph;
 
 /** rich text decorators for BlueprintHeaderView Syntax Highlighting */
 namespace HeaderViewSyntaxDecorators
@@ -67,7 +69,7 @@ protected:
 
 using FHeaderViewListItemPtr = TSharedPtr<FHeaderViewListItem>;
 
-class SBlueprintHeaderView : public SCompoundWidget
+class SBlueprintHeaderView : public SCompoundWidget, public FNotifyHook
 {
 public:
 	SLATE_BEGIN_ARGS(SBlueprintHeaderView)
@@ -82,6 +84,11 @@ public:
 	virtual FReply OnKeyDown(const FGeometry& MyGeometry, const FKeyEvent& InKeyEvent) override;
 	//~ End SWidget interface
 
+	//~ FNotifyHook interface
+		/** Handles when the settings have changed, saves to config */
+	virtual void NotifyPostChange(const FPropertyChangedEvent& PropertyChangedEvent, FProperty* PropertyThatChanged) override;
+	//~ End FNotifyHook interface
+
 private:
 	/** Gets the text for the class picker combo button */
 	FText GetClassPickerText() const;
@@ -89,6 +96,9 @@ private:
 	/** Constructs a Blueprint Class picker menu widget */
 	TSharedRef<SWidget> GetClassPickerMenuContent();
 
+	/** Constructs a DetailsView widget for the settings menu */
+	TSharedRef<SWidget> GetSettingsMenuContent();
+
 	/** Callback for class picker menu selecting a blueprint asset */
 	void OnAssetSelected(const FAssetData& SelectedAsset);
 
@@ -101,8 +111,14 @@ private:
 	/** Adds items to the list view representing all functions present in the given blueprint */
 	void PopulateFunctionItems(const UBlueprint* Blueprint);
 
+	/** Gathers all function graphs from the blueprint and sorts them according to the selected method from config */
+	void GatherFunctionGraphs(const UBlueprint* Blueprint, TArray<const UEdGraph*>& OutFunctionGraphs);
+
 	/** Adds items to the list view representing all variables present in the given blueprints */
 	void PopulateVariableItems(const UBlueprint* Blueprint);
+	
+	/** Gathers all properties from the blueprint and sorts them according to the selected method from config */
+	void GatherProperties(const UBlueprint* Blueprint, TArray<const FProperty*>& OutProperties);
 
 	/** Creates a context menu for the list view */
 	TSharedPtr<SWidget> OnContextMenuOpening();
diff --git a/Engine/Plugins/Editor/BlueprintHeaderView/Source/BlueprintHeaderView/Public/BlueprintHeaderViewSettings.h b/Engine/Plugins/Editor/BlueprintHeaderView/Source/BlueprintHeaderView/Public/BlueprintHeaderViewSettings.h
new file mode 100644
--- /dev/null
+++ b/Engine/Plugins/Editor/BlueprintHeaderView/Source/BlueprintHeaderView/Public/BlueprintHeaderViewSettings.h
@@ -0,0 +1,62 @@
+// Copyright Epic Games, Inc. All Rights Reserved.
+
+#pragma once
+
+#include "Engine/DeveloperSettings.h"
+#include "BlueprintHeaderViewSettings.generated.h"
+
+UENUM()
+enum class EHeaderViewSortMethod : uint8
+{
+	None,
+	SortByAccessSpecifier
+};
+
+USTRUCT(NotBlueprintable)
+struct FHeaderViewSyntaxColors
+{
+	GENERATED_BODY()
+	
+public:
+	UPROPERTY(config, EditAnywhere, Category="Settings")
+	FLinearColor Comment = FLinearColor(0.3f, 0.7f, 0.1f, 1.0f);
+	
+	UPROPERTY(config, EditAnywhere, Category="Settings")
+	FLinearColor Macro = FLinearColor(0.6f, 0.2f, 0.8f, 1.0f);
+	
+	UPROPERTY(config, EditAnywhere, Category="Settings")
+	FLinearColor Typename = FLinearColor::White;
+	
+	UPROPERTY(config, EditAnywhere, Category="Settings")
+	FLinearColor Identifier = FLinearColor::White;
+	
+	UPROPERTY(config, EditAnywhere, Category="Settings")
+	FLinearColor Keyword = FLinearColor(0.0f, 0.4f, 0.8f, 1.0f);
+	
+};
+
+/** Settings for the Blueprint Header View Plugin */
+UCLASS(config = EditorPerProjectUserSettings)
+class BLUEPRINTHEADERVIEW_API UBlueprintHeaderViewSettings : public UDeveloperSettings
+{
+	GENERATED_BODY()
+
+public:
+	UBlueprintHeaderViewSettings();
+
+	//~ Begin UDeveloperSettings interface
+	virtual FName GetCategoryName() const override;
+	virtual FText GetSectionText() const override;
+	virtual FName GetSectionName() const override;
+	//~ End UDeveloperSettings interface
+
+public:
+	
+	/** Syntax Highlighting colors for Blueprint Header View output */
+	UPROPERTY(config, EditAnywhere, Category="Settings")
+	FHeaderViewSyntaxColors SyntaxColors;
+
+	/** Sorting Method for Header View Functions and Properties */
+	UPROPERTY(config, EditAnywhere, Category="Settings")
+	EHeaderViewSortMethod SortMethod = EHeaderViewSortMethod::None;
+};
