BP Header View: Add context menu options to rename illegal fields

diff --git a/Engine/Plugins/Editor/BlueprintHeaderView/Source/BlueprintHeaderView/Private/HeaderViewClassListItem.cpp b/Engine/Plugins/Editor/BlueprintHeaderView/Source/BlueprintHeaderView/Private/HeaderViewClassListItem.cpp
--- a/Engine/Plugins/Editor/BlueprintHeaderView/Source/BlueprintHeaderView/Private/HeaderViewClassListItem.cpp
+++ b/Engine/Plugins/Editor/BlueprintHeaderView/Source/BlueprintHeaderView/Private/HeaderViewClassListItem.cpp
@@ -3,12 +3,32 @@
 #include "HeaderViewClassListItem.h"
 #include "Engine/Blueprint.h"
 #include "String/LineEndings.h"
+#include "Kismet2/BlueprintEditorUtils.h"
+#include "Framework/MultiBox/MultiBoxBuilder.h"
+
+#define LOCTEXT_NAMESPACE "FHeaderViewClassListItem"
 
 FHeaderViewListItemPtr FHeaderViewClassListItem::Create(TWeakObjectPtr<UBlueprint> InBlueprint)
 {
 	return MakeShareable(new FHeaderViewClassListItem(InBlueprint));
 }
 
+void FHeaderViewClassListItem::ExtendContextMenu(FMenuBuilder& InMenuBuilder, TWeakObjectPtr<UBlueprint> InBlueprint)
+{
+	if (!bIsValidName)
+	{
+		if (UBlueprint* Blueprint = InBlueprint.Get())
+		{
+			InMenuBuilder.AddEditableText(LOCTEXT("RenameBlueprint", "Rename Blueprint"),
+				LOCTEXT("RenameItemTooltip", "Renames this Blueprint\nThis Blueprint name is not a legal C++ identifier."),
+				FSlateIcon(),
+				FText::FromString(Blueprint->GetName()),
+				FOnTextCommitted::CreateSP(this, &FHeaderViewClassListItem::OnRenameTextComitted, InBlueprint)
+			);
+		}
+	}
+}
+
 FString FHeaderViewClassListItem::GetConditionalUClassSpecifiers(const UBlueprint* Blueprint) const
 {
 	TStringBuilder<256> AdditionalSpecifiers;
@@ -111,10 +131,13 @@ FHeaderViewClassListItem::FHeaderViewClassListItem(TWeakObjectPtr<UBlueprint> In
 				ParentBlueprint ? TEXT("") : Blueprint->ParentClass->GetPrefixCPP(),
 				ParentBlueprint ? *ParentBlueprint->GetName() : *Blueprint->ParentClass->GetAuthoredName());
 
+			bIsValidName = IsValidCPPIdentifier(BlueprintName);
+			const FString& NameDecorator = bIsValidName ? HeaderViewSyntaxDecorators::TypenameDecorator : HeaderViewSyntaxDecorators::ErrorDecorator;
+
 			RawItemString += FString::Printf(TEXT("\nclass %s : public %s\n{\n\tGENERATED_BODY()"), *BlueprintName, *ParentClassName);
 			RichTextString += FString::Printf(TEXT("\n<%s>class</> <%s>%s</> : <%s>public</> <%s>%s</>\n{\n\t<%s>GENERATED_BODY</>()"), 
 				*HeaderViewSyntaxDecorators::KeywordDecorator,
-				*HeaderViewSyntaxDecorators::TypenameDecorator, 
+				*NameDecorator, 
 				*BlueprintName,
 				*HeaderViewSyntaxDecorators::KeywordDecorator,
 				*HeaderViewSyntaxDecorators::TypenameDecorator,
@@ -129,3 +152,20 @@ FHeaderViewClassListItem::FHeaderViewClassListItem(TWeakObjectPtr<UBlueprint> In
 	}
 }
 
+
+void FHeaderViewClassListItem::OnRenameTextComitted(const FText& CommittedText, ETextCommit::Type TextCommitType, TWeakObjectPtr<UBlueprint> InBlueprint)
+{
+	if (TextCommitType == ETextCommit::OnEnter)
+	{
+		if (UBlueprint* Blueprint = InBlueprint.Get())
+		{
+			if (IsValidCPPIdentifier(CommittedText.ToString()))
+			{
+				Blueprint->Rename(*CommittedText.ToString());
+				FBlueprintEditorUtils::MarkBlueprintAsStructurallyModified(Blueprint);
+			}
+		}
+	}
+}
+
+#undef LOCTEXT_NAMESPACE
diff --git a/Engine/Plugins/Editor/BlueprintHeaderView/Source/BlueprintHeaderView/Private/HeaderViewClassListItem.h b/Engine/Plugins/Editor/BlueprintHeaderView/Source/BlueprintHeaderView/Private/HeaderViewClassListItem.h
--- a/Engine/Plugins/Editor/BlueprintHeaderView/Source/BlueprintHeaderView/Private/HeaderViewClassListItem.h
+++ b/Engine/Plugins/Editor/BlueprintHeaderView/Source/BlueprintHeaderView/Private/HeaderViewClassListItem.h
@@ -5,6 +5,8 @@
 #include "CoreMinimal.h"
 #include "SBlueprintHeaderView.h"
 
+class FMenuBuilder;
+
 /** A header view list item that displays the class declaration */
 struct FHeaderViewClassListItem : public FHeaderViewListItem
 {
@@ -12,9 +14,17 @@ public:
 	/** Creates a list item for the Header view representing a class declaration for the given blueprint */
 	static FHeaderViewListItemPtr Create(TWeakObjectPtr<UBlueprint> InBlueprint);
 
+	//~ FHeaderViewListItem Interface
+	virtual void ExtendContextMenu(FMenuBuilder& InMenuBuilder, TWeakObjectPtr<UBlueprint> InBlueprint) override;
+	//~ End FHeaderViewListItem Interface
 protected:
 	FString GetConditionalUClassSpecifiers(const UBlueprint* Blueprint) const;
 
 	FHeaderViewClassListItem(TWeakObjectPtr<UBlueprint> InBlueprint);
 
+	void OnRenameTextComitted(const FText& CommittedText, ETextCommit::Type TextCommitType, TWeakObjectPtr<UBlueprint> InBlueprint);
+
+protected:
+	/** Whether this class name is valid C++ (no spaces, special chars, etc) */
+	bool bIsValidName = true;
 };
\ No newline at end of file
diff --git a/Engine/Plugins/Editor/BlueprintHeaderView/Source/BlueprintHeaderView/Private/HeaderViewFunctionListItem.cpp b/Engine/Plugins/Editor/BlueprintHeaderView/Source/BlueprintHeaderView/Private/HeaderViewFunctionListItem.cpp
--- a/Engine/Plugins/Editor/BlueprintHeaderView/Source/BlueprintHeaderView/Private/HeaderViewFunctionListItem.cpp
+++ b/Engine/Plugins/Editor/BlueprintHeaderView/Source/BlueprintHeaderView/Private/HeaderViewFunctionListItem.cpp
@@ -6,12 +6,38 @@
 #include "Kismet2/BlueprintEditorUtils.h"
 #include "K2Node_FunctionEntry.h"
 #include "Misc/EngineVersion.h"
+#include "Framework/MultiBox/MultiBoxBuilder.h"
+
+#define LOCTEXT_NAMESPACE "FHeaderViewFunctionListItem"
 
 FHeaderViewListItemPtr FHeaderViewFunctionListItem::Create(const UK2Node_FunctionEntry* FunctionEntry)
 {
 	return MakeShareable(new FHeaderViewFunctionListItem(FunctionEntry));
 }
 
+void FHeaderViewFunctionListItem::ExtendContextMenu(FMenuBuilder& InMenuBuilder, TWeakObjectPtr<UBlueprint> InBlueprint)
+{
+	if (!IllegalName.IsNone())
+	{
+		InMenuBuilder.AddEditableText(LOCTEXT("RenameItem", "Rename Function"),
+			LOCTEXT("RenameItemTooltip", "Renames this function in the Blueprint\nThis function name is not a legal C++ identifier."),
+			FSlateIcon(),
+			FText::FromName(IllegalName),
+			FOnTextCommitted::CreateSP(this, &FHeaderViewFunctionListItem::OnRenameFunctionTextCommitted, InBlueprint, GraphName)
+		);
+	}
+
+	for (FName IllegalParam : IllegalParameters)
+	{
+		InMenuBuilder.AddEditableText(LOCTEXT("RenameParm", "Rename Parameter"),
+			LOCTEXT("RenameParmTooltip", "Renames this function parameter in the Blueprint\nThis parameter name is not a legal C++ identifier."),
+			FSlateIcon(),
+			FText::FromName(IllegalParam),
+			FOnTextCommitted::CreateSP(this, &FHeaderViewFunctionListItem::OnRenameParameterTextCommitted, InBlueprint, GraphName, IllegalParam)
+		);
+	}
+}
+
 FString FHeaderViewFunctionListItem::GetConditionalUFunctionSpecifiers(const UFunction* SigFunction) const
 {
 	check(SigFunction);
@@ -115,17 +141,26 @@ void FHeaderViewFunctionListItem::AppendFunctionParameters(const UFunction* Sign
 			RichTextString.Append(FString::Printf(TEXT("<%s>UPARAM</>(ref) "), *HeaderViewSyntaxDecorators::MacroDecorator));
 		}
 
+		const FString Typename = GetCPPTypenameForProperty(*ParmIt);
+		const FString ParmName = ParmIt->GetAuthoredName();
+		
+		const bool bIsValidName = IsValidCPPIdentifier(ParmName);
+		const FString* IdentifierDecorator = &HeaderViewSyntaxDecorators::IdentifierDecorator;
+		if (!bIsValidName)
+		{
+			IllegalParameters.Emplace(ParmName);
+			IdentifierDecorator = &HeaderViewSyntaxDecorators::ErrorDecorator;
+		}
+
 		if (ParmIt->HasAnyPropertyFlags(CPF_OutParm | CPF_ReferenceParm))
 		{
-			const FString Typename = GetCPPTypenameForProperty(*ParmIt);
-			const FString ParmName = ParmIt->GetAuthoredName();
 			RawItemString.Append(FString::Printf(TEXT("%s& %s"), *Typename, *ParmName));
-			RichTextString.Append(FString::Printf(TEXT("<%s>%s</>& <%s>%s</>"), *HeaderViewSyntaxDecorators::TypenameDecorator, *Typename, *HeaderViewSyntaxDecorators::IdentifierDecorator, *ParmName));
+			RichTextString.Append(FString::Printf(TEXT("<%s>%s</>& <%s>%s</>"), *HeaderViewSyntaxDecorators::TypenameDecorator, *Typename, **IdentifierDecorator, *ParmName));
 		}
 		else
 		{
-			RawItemString.Append(FString::Printf(TEXT("%s %s"), *ParmIt->GetCPPType(), *ParmIt->GetAuthoredName()));
-			RichTextString.Append(FString::Printf(TEXT("<%s>%s</> <%s>%s</>"), *HeaderViewSyntaxDecorators::TypenameDecorator, *ParmIt->GetCPPType(), *HeaderViewSyntaxDecorators::IdentifierDecorator, *ParmIt->GetAuthoredName()));
+			RawItemString.Append(FString::Printf(TEXT("%s %s"), *Typename, *ParmName));
+			RichTextString.Append(FString::Printf(TEXT("<%s>%s</> <%s>%s</>"), *HeaderViewSyntaxDecorators::TypenameDecorator, *Typename, **IdentifierDecorator, *ParmName));
 		}
 
 		++ParamIdx;
@@ -136,6 +171,7 @@ FHeaderViewFunctionListItem::FHeaderViewFunctionListItem(const UK2Node_FunctionE
 {
 	check(FunctionEntry);
 
+	GraphName = FunctionEntry->GetGraph()->GetFName();
 	RawItemString.Reserve(512);
 	RichTextString.Reserve(512);
 
@@ -198,19 +234,26 @@ FHeaderViewFunctionListItem::FHeaderViewFunctionListItem(const UK2Node_FunctionE
 				FunctionName = FunctionEntry->CustomGeneratedFunctionName.ToString();
 			}
 
+			const bool bValidName = IsValidCPPIdentifier(FunctionName);
+			const FString* IdentifierDecorator = &HeaderViewSyntaxDecorators::IdentifierDecorator;
+			if (!bValidName)
+			{
+				IllegalName = ResolvedFunc->GetFName();
+				IdentifierDecorator = &HeaderViewSyntaxDecorators::ErrorDecorator;
+			}
+
 			if (FProperty* ReturnProperty = ResolvedFunc->GetReturnProperty())
 			{
 				const FString Typename = ReturnProperty->GetCPPType();
 				RawItemString += FString::Printf(TEXT("\n%s %s("), *Typename, *FunctionName);
-				RichTextString += FString::Printf(TEXT("\n<%s>%s</> <%s>%s</>("), *HeaderViewSyntaxDecorators::TypenameDecorator, *Typename, *HeaderViewSyntaxDecorators::IdentifierDecorator, *FunctionName);
+				RichTextString += FString::Printf(TEXT("\n<%s>%s</> <%s>%s</>("), *HeaderViewSyntaxDecorators::TypenameDecorator, *Typename, **IdentifierDecorator, *FunctionName);
 			}
 			else
 			{
 				RawItemString += FString::Printf(TEXT("\nvoid %s("), *FunctionName);
-				RichTextString += FString::Printf(TEXT("\n<%s>void</> <%s>%s</>("), *HeaderViewSyntaxDecorators::KeywordDecorator, *HeaderViewSyntaxDecorators::IdentifierDecorator, *FunctionName);
+				RichTextString += FString::Printf(TEXT("\n<%s>void</> <%s>%s</>("), *HeaderViewSyntaxDecorators::KeywordDecorator, **IdentifierDecorator, *FunctionName);
 			}
 
-
 			AppendFunctionParameters(ResolvedFunc);
 
 			if (FunctionEntry->GetFunctionFlags() & FUNC_Const)
@@ -236,3 +279,67 @@ FHeaderViewFunctionListItem::FHeaderViewFunctionListItem(const UK2Node_FunctionE
 		UE::String::ToHostLineEndingsInline(RichTextString);
 	}
 }
+
+void FHeaderViewFunctionListItem::OnRenameFunctionTextCommitted(const FText& CommittedText, ETextCommit::Type TextCommitType, TWeakObjectPtr<UBlueprint> WeakBlueprint, FName OldGraphName)
+{
+	if (TextCommitType == ETextCommit::OnEnter)
+	{
+		if (UBlueprint* Blueprint = WeakBlueprint.Get())
+		{
+			const FString& CommittedString = CommittedText.ToString();
+			if (IsValidCPPIdentifier(CommittedString))
+			{
+				const TObjectPtr<UEdGraph>* FunctionGraph = Blueprint->FunctionGraphs.FindByPredicate([OldGraphName](const TObjectPtr<UEdGraph>& Graph)
+					{
+						return Graph->GetFName() == OldGraphName;
+					});
+
+				if (FunctionGraph)
+				{
+					FBlueprintEditorUtils::RenameGraph(FunctionGraph->Get(), CommittedString);
+				}
+				else
+				{
+					UE_LOG(LogBlueprintHeaderView, Warning, TEXT("Could not find Function Graph named \"%s\" in Blueprint \"%s\""), *OldGraphName.ToString(), *Blueprint->GetName());
+				}
+			}
+		}
+	}
+}
+
+void FHeaderViewFunctionListItem::OnRenameParameterTextCommitted(const FText& CommittedText, ETextCommit::Type TextCommitType, TWeakObjectPtr<UBlueprint> WeakBlueprint, FName OldGraphName, FName OldParamName)
+{
+	if (TextCommitType == ETextCommit::OnEnter)
+	{
+		if (UBlueprint* Blueprint = WeakBlueprint.Get())
+		{
+			const FString& CommittedString = CommittedText.ToString();
+			if (IsValidCPPIdentifier(CommittedString))
+			{
+				const TObjectPtr<UEdGraph>* FunctionGraph = Blueprint->FunctionGraphs.FindByPredicate([OldGraphName](const TObjectPtr<UEdGraph>& Graph)
+					{
+						return Graph->GetFName() == OldGraphName;
+					});
+
+				if (FunctionGraph)
+				{
+					TArray<UK2Node_FunctionEntry*> Entry;
+					FunctionGraph->Get()->GetNodesOfClass<UK2Node_FunctionEntry>(Entry);
+					if (ensureMsgf(Entry.Num() == 1 && Entry[0], TEXT("Function Graph \"%s\" in Blueprint \"%s\" does not have exactly one Entry Node"), *FunctionGraph->Get()->GetName(), *Blueprint->GetName()))
+					{
+						if (Entry[0]->RenameUserDefinedPin(OldParamName, FName(CommittedString)) == ERenamePinResult_Success)
+						{
+							FBlueprintEditorUtils::MarkBlueprintAsStructurallyModified(Blueprint);
+						}
+					}
+				}
+				else
+				{
+					UE_LOG(LogBlueprintHeaderView, Warning, TEXT("Could not find Function Graph named \"%s\" in Blueprint \"%s\""), *OldGraphName.ToString(), *Blueprint->GetName());
+				}
+			}
+		}
+	}
+}
+
+#undef LOCTEXT_NAMESPACE
diff --git a/Engine/Plugins/Editor/BlueprintHeaderView/Source/BlueprintHeaderView/Private/HeaderViewFunctionListItem.h b/Engine/Plugins/Editor/BlueprintHeaderView/Source/BlueprintHeaderView/Private/HeaderViewFunctionListItem.h
--- a/Engine/Plugins/Editor/BlueprintHeaderView/Source/BlueprintHeaderView/Private/HeaderViewFunctionListItem.h
+++ b/Engine/Plugins/Editor/BlueprintHeaderView/Source/BlueprintHeaderView/Private/HeaderViewFunctionListItem.h
@@ -13,6 +13,10 @@ struct FHeaderViewFunctionListItem : public FHeaderViewListItem
 	/** Creates a list item for the Header view representing a function declaration for the given blueprint function */
 	static FHeaderViewListItemPtr Create(const UK2Node_FunctionEntry* FunctionEntry);
 
+	//~ FHeaderViewListItem Interface
+	virtual void ExtendContextMenu(FMenuBuilder& InMenuBuilder, TWeakObjectPtr<UBlueprint> InBlueprint) override;
+	//~ End FHeaderViewListItem Interface
+
 protected:
 	FHeaderViewFunctionListItem(const UK2Node_FunctionEntry* FunctionEntry);
 
@@ -22,4 +26,16 @@ protected:
 	/** Adds Function parameters to the RichText and PlainText strings */
 	void AppendFunctionParameters(const UFunction* SignatureFunction);
 
+	void OnRenameFunctionTextCommitted(const FText& CommittedText, ETextCommit::Type TextCommitType, TWeakObjectPtr<UBlueprint> WeakBlueprint, FName OldGraphName);
+	void OnRenameParameterTextCommitted(const FText& CommittedText, ETextCommit::Type TextCommitType, TWeakObjectPtr<UBlueprint> WeakBlueprint, FName OldGraphName, FName OldParamName);
+
+protected:
+	/** None if the function name is legal C++, else the name of the function */
+	FName IllegalName = NAME_None;
+
+	/** Name of the Function Graph this item represents */
+	FName GraphName = NAME_None;
+
+	/** Names of any function parameters that are not legal C++ */
+	TArray<FName> IllegalParameters;
 };
\ No newline at end of file
diff --git a/Engine/Plugins/Editor/BlueprintHeaderView/Source/BlueprintHeaderView/Private/HeaderViewVariableListItem.cpp b/Engine/Plugins/Editor/BlueprintHeaderView/Source/BlueprintHeaderView/Private/HeaderViewVariableListItem.cpp
--- a/Engine/Plugins/Editor/BlueprintHeaderView/Source/BlueprintHeaderView/Private/HeaderViewVariableListItem.cpp
+++ b/Engine/Plugins/Editor/BlueprintHeaderView/Source/BlueprintHeaderView/Private/HeaderViewVariableListItem.cpp
@@ -5,12 +5,29 @@
 #include "EdGraphSchema_K2.h"
 #include "Engine/Blueprint.h"
 #include "String/LineEndings.h"
+#include "Kismet2/BlueprintEditorUtils.h"
+#include "Framework/MultiBox/MultiBoxBuilder.h"
+
+#define LOCTEXT_NAMESPACE "FHeaderViewVariableListItem"
 
 FHeaderViewListItemPtr FHeaderViewVariableListItem::Create(const FBPVariableDescription* VariableDesc, const FProperty& VarProperty)
 {
 	return MakeShareable(new FHeaderViewVariableListItem(VariableDesc, VarProperty));
 }
 
+void FHeaderViewVariableListItem::ExtendContextMenu(FMenuBuilder& InMenuBuilder, TWeakObjectPtr<UBlueprint> InBlueprint)
+{
+	if (!IllegalName.IsNone())
+	{
+		InMenuBuilder.AddEditableText(LOCTEXT("RenameItem", "Rename Variable"),
+			LOCTEXT("RenameItemTooltip", "Renames this variable in the Blueprint\nThis variable name is not a legal C++ identifier."),
+			FSlateIcon(),
+			FText::FromName(IllegalName),
+			FOnTextCommitted::CreateSP(this, &FHeaderViewVariableListItem::OnRenameTextCommitted, InBlueprint)
+		);
+	}
+}
+
 FHeaderViewVariableListItem::FHeaderViewVariableListItem(const FBPVariableDescription* VariableDesc, const FProperty& VarProperty)
 {
 	// Format comment
@@ -88,12 +105,21 @@ FHeaderViewVariableListItem::FHeaderViewVariableListItem(const FBPVariableDescri
 	{
 		const FString Typename = GetCPPTypenameForProperty(&VarProperty, /*bIsMemberProperty=*/true);
 		const FString VarName = VarProperty.GetAuthoredName();
+		const bool bLegalName = IsValidCPPIdentifier(VarName);
+
+		const FString* IdentifierDecorator = &HeaderViewSyntaxDecorators::IdentifierDecorator;
+		if (!bLegalName)
+		{
+			IllegalName = FName(VarName);
+			IdentifierDecorator = &HeaderViewSyntaxDecorators::ErrorDecorator;
+		}
+
 
 		RawItemString += FString::Printf(TEXT("\n%s %s;"), *Typename, *VarName);
 		RichTextString += FString::Printf(TEXT("\n<%s>%s</> <%s>%s</>;"), 
 			*HeaderViewSyntaxDecorators::TypenameDecorator,
 			*Typename.Replace(TEXT("<"), TEXT("&lt;")).Replace(TEXT(">"), TEXT("&gt;")),
-			*HeaderViewSyntaxDecorators::IdentifierDecorator,
+			**IdentifierDecorator,
 			*VarName);
 	}
 
@@ -305,3 +331,20 @@ FString FHeaderViewVariableListItem::GetOwningClassName(const FProperty& VarProp
 
 	return TEXT("");
 }
+
+void FHeaderViewVariableListItem::OnRenameTextCommitted(const FText& CommittedText, ETextCommit::Type TextCommitType, TWeakObjectPtr<UBlueprint> WeakBlueprint)
+{
+	if (TextCommitType == ETextCommit::OnEnter)
+	{
+		if (UBlueprint* Blueprint = WeakBlueprint.Get())
+		{
+			const FString& CommittedString = CommittedText.ToString();
+			if (IsValidCPPIdentifier(CommittedString))
+			{
+				FBlueprintEditorUtils::RenameMemberVariable(Blueprint, IllegalName, FName(CommittedString));
+			}
+		}
+	}
+}
+
+#undef LOCTEXT_NAMESPACE
diff --git a/Engine/Plugins/Editor/BlueprintHeaderView/Source/BlueprintHeaderView/Private/HeaderViewVariableListItem.h b/Engine/Plugins/Editor/BlueprintHeaderView/Source/BlueprintHeaderView/Private/HeaderViewVariableListItem.h
--- a/Engine/Plugins/Editor/BlueprintHeaderView/Source/BlueprintHeaderView/Private/HeaderViewVariableListItem.h
+++ b/Engine/Plugins/Editor/BlueprintHeaderView/Source/BlueprintHeaderView/Private/HeaderViewVariableListItem.h
@@ -13,6 +13,10 @@ struct FHeaderViewVariableListItem : public FHeaderViewListItem
 	/** Creates a list item for the Header view representing a variable declaration for the given blueprint variable */
 	static FHeaderViewListItemPtr Create(const FBPVariableDescription* VariableDesc, const FProperty& VarProperty);
 
+	//~ FHeaderViewListItem Interface
+	virtual void ExtendContextMenu(FMenuBuilder& InMenuBuilder, TWeakObjectPtr<UBlueprint> Blueprint) override;
+	//~ End FHeaderViewListItem Interface
+
 protected:
 	FHeaderViewVariableListItem(const FBPVariableDescription* VariableDesc, const FProperty& VarProperty);
 
@@ -25,4 +29,9 @@ protected:
 	/** Returns the name of the owning class */
 	FString GetOwningClassName(const FProperty& VarProperty) const;
 
+	void OnRenameTextCommitted(const FText& CommittedText, ETextCommit::Type TextCommitType, TWeakObjectPtr<UBlueprint> WeakBlueprint);
+protected:
+
+	/** None if the name is legal, else holds the name of the variable */
+	FName IllegalName = NAME_None;
 };
\ No newline at end of file
diff --git a/Engine/Plugins/Editor/BlueprintHeaderView/Source/BlueprintHeaderView/Private/SBlueprintHeaderView.cpp b/Engine/Plugins/Editor/BlueprintHeaderView/Source/BlueprintHeaderView/Private/SBlueprintHeaderView.cpp
--- a/Engine/Plugins/Editor/BlueprintHeaderView/Source/BlueprintHeaderView/Private/SBlueprintHeaderView.cpp
+++ b/Engine/Plugins/Editor/BlueprintHeaderView/Source/BlueprintHeaderView/Private/SBlueprintHeaderView.cpp
@@ -22,14 +22,18 @@
 #include "String/LineEndings.h"
 #include "HAL/PlatformApplicationMisc.h"
 #include "BlueprintHeaderViewSettings.h"
+#include "Internationalization/Regex.h"
 
 #define LOCTEXT_NAMESPACE "SBlueprintHeaderView"
 
+DEFINE_LOG_CATEGORY(LogBlueprintHeaderView)
+
 // HeaderViewSyntaxDecorators /////////////////////////////////////////////////
 
 namespace HeaderViewSyntaxDecorators
 {
 	const FString CommentDecorator = TEXT("comment");
+	const FString ErrorDecorator = TEXT("error");
 	const FString IdentifierDecorator = TEXT("identifier");
 	const FString KeywordDecorator = TEXT("keyword");
 	const FString MacroDecorator = TEXT("macro");
@@ -103,6 +107,7 @@ TSharedRef<SWidget> FHeaderViewListItem::GenerateWidgetForItem()
 			.Text(FText::FromString(RichTextString))
 			.TextStyle(FEditorStyle::Get(), "Log.Normal")
 			+SRichTextBlock::Decorator(FHeaderViewSyntaxDecorator::Create(HeaderViewSyntaxDecorators::CommentDecorator, SyntaxColors.Comment))
+			+SRichTextBlock::Decorator(FHeaderViewSyntaxDecorator::Create(HeaderViewSyntaxDecorators::ErrorDecorator, SyntaxColors.Error))
 			+SRichTextBlock::Decorator(FHeaderViewSyntaxDecorator::Create(HeaderViewSyntaxDecorators::IdentifierDecorator, SyntaxColors.Identifier))
 			+SRichTextBlock::Decorator(FHeaderViewSyntaxDecorator::Create(HeaderViewSyntaxDecorators::KeywordDecorator, SyntaxColors.Keyword))
 			+SRichTextBlock::Decorator(FHeaderViewSyntaxDecorator::Create(HeaderViewSyntaxDecorators::MacroDecorator, SyntaxColors.Macro))
@@ -174,6 +179,15 @@ FString FHeaderViewListItem::GetCPPTypenameForProperty(const FProperty* InProper
 	}
 }
 
+bool FHeaderViewListItem::IsValidCPPIdentifier(const FString& InIdentifier)
+{
+	// Only matches strings that start with a Letter or Underscore with any number of letters, digits or underscores following
+	// _Legal, Legal0, legal_0
+	// 0Illegal, Illegal&, Illegal-0
+	static const FRegexPattern RegexPattern = FRegexPattern(TEXT("^[A-Za-z_]\\w*$"));
+	return FRegexMatcher(RegexPattern, InIdentifier).FindNext();
+}
+
 // SBlueprintHeaderView ///////////////////////////////////////////////////////
 
 void SBlueprintHeaderView::Construct(const FArguments& InArgs)
@@ -341,7 +355,16 @@ void SBlueprintHeaderView::OnAssetSelected(const FAssetData& SelectedAsset)
 {
 	ClassPickerComboButton->SetIsOpen(false);
 
-	SelectedBlueprint = Cast<UBlueprint>(SelectedAsset.GetAsset());
+	if (UBlueprint* OldBlueprint = SelectedBlueprint.Get())
+	{
+		OldBlueprint->OnChanged().RemoveAll(this);
+	}
+
+	if (UBlueprint* Blueprint = Cast<UBlueprint>(SelectedAsset.GetAsset()))
+	{
+		Blueprint->OnChanged().AddRaw(this, &SBlueprintHeaderView::OnBlueprintChanged);
+		SelectedBlueprint = Blueprint;
+	}
 
 	RepopulateListView();
 }
@@ -531,9 +554,21 @@ TSharedPtr<SWidget> SBlueprintHeaderView::OnContextMenuOpening()
 	MenuBuilder.AddMenuEntry(FGenericCommands::Get().SelectAll);
 	MenuBuilder.AddMenuEntry(FGenericCommands::Get().Copy);
 
+	TArray<FHeaderViewListItemPtr> SelectedListItems;
+	ListView->GetSelectedItems(SelectedListItems);
+	if (SelectedListItems.Num() == 1)
+	{
+		SelectedListItems[0]->ExtendContextMenu(MenuBuilder, SelectedBlueprint);
+	}
+
 	return MenuBuilder.MakeWidget();
 }
 
+void SBlueprintHeaderView::OnBlueprintChanged(UBlueprint* InBlueprint)
+{
+	RepopulateListView();
+}
+
 void SBlueprintHeaderView::OnCopy() const
 {
 	const int32 StringReserveSize = 2048;
diff --git a/Engine/Plugins/Editor/BlueprintHeaderView/Source/BlueprintHeaderView/Private/SBlueprintHeaderView.h b/Engine/Plugins/Editor/BlueprintHeaderView/Source/BlueprintHeaderView/Private/SBlueprintHeaderView.h
--- a/Engine/Plugins/Editor/BlueprintHeaderView/Source/BlueprintHeaderView/Private/SBlueprintHeaderView.h
+++ b/Engine/Plugins/Editor/BlueprintHeaderView/Source/BlueprintHeaderView/Private/SBlueprintHeaderView.h
@@ -16,10 +16,13 @@ class STableViewBase;
 class FUICommandList;
 class UEdGraph;
 
+DECLARE_LOG_CATEGORY_EXTERN(LogBlueprintHeaderView, Log, All)
+
 /** rich text decorators for BlueprintHeaderView Syntax Highlighting */
 namespace HeaderViewSyntaxDecorators
 {
 	extern const FString CommentDecorator;
+	extern const FString ErrorDecorator;
 	extern const FString IdentifierDecorator;
 	extern const FString KeywordDecorator;
 	extern const FString MacroDecorator;
@@ -29,6 +32,8 @@ namespace HeaderViewSyntaxDecorators
 /** A base class for List Items in the Header View */
 struct FHeaderViewListItem : public TSharedFromThis<FHeaderViewListItem>
 {
+	virtual ~FHeaderViewListItem() {};
+
 	/** Creates the widget for this list item */
 	TSharedRef<SWidget> GenerateWidgetForItem();
 
@@ -38,6 +43,9 @@ struct FHeaderViewListItem : public TSharedFromThis<FHeaderViewListItem>
 	/** Returns the raw item text for copy actions */
 	const FString& GetRawItemString() const { return RawItemString; }
 
+	/** Allows the item to add items to the context menu if it is the only item selected */
+	virtual void ExtendContextMenu(FMenuBuilder& InMenuBuilder, TWeakObjectPtr<UBlueprint> InBlueprint) {}
+
 protected:
 	/** Empty base constructor hidden from public */
 	FHeaderViewListItem() {};
@@ -58,6 +66,9 @@ protected:
 	 */
 	static FString GetCPPTypenameForProperty(const FProperty* InProperty, bool bIsMemberProperty = false);
 
+	/** Checks whether a string is a valid C++ identifier */
+	static bool IsValidCPPIdentifier(const FString& InIdentifier);
+
 protected:
 	/** A rich text representation of the item, including syntax highlighting and errors */
 	FString RichTextString;
@@ -123,6 +134,9 @@ private:
 	/** Creates a context menu for the list view */
 	TSharedPtr<SWidget> OnContextMenuOpening();
 
+	/** Callback for when the selected blueprint is modified */
+	void OnBlueprintChanged(UBlueprint* InBlueprint);
+
 	/** UI Command Functions */
 	void OnCopy() const;
 	bool CanCopy() const;
diff --git a/Engine/Plugins/Editor/BlueprintHeaderView/Source/BlueprintHeaderView/Public/BlueprintHeaderViewSettings.h b/Engine/Plugins/Editor/BlueprintHeaderView/Source/BlueprintHeaderView/Public/BlueprintHeaderViewSettings.h
--- a/Engine/Plugins/Editor/BlueprintHeaderView/Source/BlueprintHeaderView/Public/BlueprintHeaderViewSettings.h
+++ b/Engine/Plugins/Editor/BlueprintHeaderView/Source/BlueprintHeaderView/Public/BlueprintHeaderViewSettings.h
@@ -21,6 +21,9 @@ public:
 	UPROPERTY(config, EditAnywhere, Category="Settings")
 	FLinearColor Comment = FLinearColor(0.3f, 0.7f, 0.1f, 1.0f);
 	
+	UPROPERTY(config, EditAnywhere, Category="Settings")
+	FLinearColor Error = FLinearColor(1.0f, 0.0f, 0.0f, 1.0f);
+	
 	UPROPERTY(config, EditAnywhere, Category="Settings")
 	FLinearColor Macro = FLinearColor(0.6f, 0.2f, 0.8f, 1.0f);
 	
