BP Header View: Add Context Menu action to Blueprint Assets to open the Header View

diff --git a/Engine/Plugins/Editor/BlueprintHeaderView/Source/BlueprintHeaderView/Private/BlueprintHeaderView.cpp b/Engine/Plugins/Editor/BlueprintHeaderView/Source/BlueprintHeaderView/Private/BlueprintHeaderView.cpp
--- a/Engine/Plugins/Editor/BlueprintHeaderView/Source/BlueprintHeaderView/Private/BlueprintHeaderView.cpp
+++ b/Engine/Plugins/Editor/BlueprintHeaderView/Source/BlueprintHeaderView/Private/BlueprintHeaderView.cpp
@@ -12,6 +12,12 @@
 #include "Editor/WorkspaceMenuStructure/Public/WorkspaceMenuStructureModule.h"
 #include "Widgets/Docking/SDockTab.h"
 #include "SBlueprintHeaderView.h"
+#include "ContentBrowserModule.h"
+#include "Framework/Commands/UIAction.h"
+#include "Engine/Blueprint.h"
+#include "Framework/MultiBox/MultiBoxBuilder.h"
+
+#define LOCTEXT_NAMESPACE "BlueprintHeaderViewApp"
 
 namespace BlueprintHeaderViewModule
 {
@@ -30,16 +36,78 @@ namespace BlueprintHeaderViewModule
 void FBlueprintHeaderViewModule::StartupModule()
 {
 	FGlobalTabmanager::Get()->RegisterNomadTabSpawner(BlueprintHeaderViewModule::HeaderViewTabName, FOnSpawnTab::CreateStatic(&BlueprintHeaderViewModule::CreateHeaderViewTab))
-		.SetDisplayName(NSLOCTEXT("BlueprintHeaderViewApp", "TabTitle", "Blueprint Header View"))
-		.SetTooltipText(NSLOCTEXT("BlueprintHeaderViewApp", "TooltipText", "Displays a Blueprint Class in C++ Header format."))
+		.SetDisplayName(LOCTEXT("TabTitle", "Blueprint Header View"))
+		.SetTooltipText(LOCTEXT("TooltipText", "Displays a Blueprint Class in C++ Header format."))
 		.SetGroup(WorkspaceMenu::GetMenuStructure().GetToolsCategory())
 		.SetIcon(FSlateIcon(FEditorStyle::GetStyleSetName(), "ClassIcon.Class"));
+
+	SetupContentBrowserContextMenuExtender();
 }
 
 void FBlueprintHeaderViewModule::ShutdownModule()
 {
-	// This function may be called during shutdown to clean up your module.  For modules that support dynamic reloading,
-	// we call this function before unloading the module.
+	if (ContentBrowserExtenderDelegateHandle.IsValid())
+	{
+		FContentBrowserModule& ContentBrowserModule = FModuleManager::LoadModuleChecked<FContentBrowserModule>("ContentBrowser");
+		TArray< FContentBrowserMenuExtender_SelectedAssets >& CBMenuExtenderDelegates = ContentBrowserModule.GetAllAssetViewContextMenuExtenders();
+		CBMenuExtenderDelegates.RemoveAll([ContentBrowserExtenderDelegateHandle=ContentBrowserExtenderDelegateHandle](const FContentBrowserMenuExtender_SelectedAssets& Delegate)
+			{
+				return Delegate.GetHandle() == ContentBrowserExtenderDelegateHandle;
+			});
+	}
 }
 	
-IMPLEMENT_MODULE(FBlueprintHeaderViewModule, BlueprintHeaderView)
\ No newline at end of file
+bool FBlueprintHeaderViewModule::IsClassHeaderViewSupported(const UClass* InClass)
+{
+	return InClass == UBlueprint::StaticClass();
+}
+
+void FBlueprintHeaderViewModule::OpenHeaderViewForAsset(FAssetData InAssetData)
+{
+	TSharedPtr<SDockTab> HeaderViewTab = FGlobalTabmanager::Get()->TryInvokeTab(BlueprintHeaderViewModule::HeaderViewTabName);
+
+	if (HeaderViewTab.IsValid())
+	{
+		TSharedRef<SWidget> HeaderViewContentWidget = HeaderViewTab->GetContent();
+		if (HeaderViewContentWidget->GetWidgetClass().GetWidgetType() == SBlueprintHeaderView::StaticWidgetClass().GetWidgetType())
+		{
+			StaticCastSharedRef<SBlueprintHeaderView>(HeaderViewContentWidget)->OnAssetSelected(InAssetData);
+		}
+	}
+}
+
+void FBlueprintHeaderViewModule::SetupContentBrowserContextMenuExtender()
+{
+	FContentBrowserModule& ContentBrowserModule = FModuleManager::LoadModuleChecked<FContentBrowserModule>("ContentBrowser");
+	TArray< FContentBrowserMenuExtender_SelectedAssets >& CBMenuExtenderDelegates = ContentBrowserModule.GetAllAssetViewContextMenuExtenders();
+	CBMenuExtenderDelegates.Add(FContentBrowserMenuExtender_SelectedAssets::CreateStatic(&FBlueprintHeaderViewModule::OnExtendContentBrowserAssetSelectionMenu));
+	ContentBrowserExtenderDelegateHandle = CBMenuExtenderDelegates.Last().GetHandle();
+}
+
+TSharedRef<FExtender> FBlueprintHeaderViewModule::OnExtendContentBrowserAssetSelectionMenu(const TArray<FAssetData>& SelectedAssets)
+{
+	TSharedRef<FExtender> Extender = MakeShared<FExtender>();
+
+	if (SelectedAssets.Num() == 1)
+	{
+		if (IsClassHeaderViewSupported(SelectedAssets[0].GetClass()))
+		{
+			Extender->AddMenuExtension("GetAssetActions", EExtensionHook::After, nullptr, FMenuExtensionDelegate::CreateLambda(
+				[SelectedAssets](FMenuBuilder& MenuBuilder) {
+					MenuBuilder.AddMenuEntry(
+						LOCTEXT("OpenHeaderView", "Display in Blueprint Header View"),
+						LOCTEXT("OpenHeaderViewTooltip", "Opens this Blueprint in the Blueprint Header View"),
+						FSlateIcon(FEditorStyle::GetStyleSetName(), "ClassIcon.Class"),
+						FUIAction(FExecuteAction::CreateStatic(&FBlueprintHeaderViewModule::OpenHeaderViewForAsset, SelectedAssets[0]))
+						);
+				})
+			);
+		}
+	}
+
+	return Extender;
+}
+
+#undef LOCTEXT_NAMESPACE
+
+IMPLEMENT_MODULE(FBlueprintHeaderViewModule, BlueprintHeaderView)
diff --git a/Engine/Plugins/Editor/BlueprintHeaderView/Source/BlueprintHeaderView/Public/BlueprintHeaderView.h b/Engine/Plugins/Editor/BlueprintHeaderView/Source/BlueprintHeaderView/Public/BlueprintHeaderView.h
--- a/Engine/Plugins/Editor/BlueprintHeaderView/Source/BlueprintHeaderView/Public/BlueprintHeaderView.h
+++ b/Engine/Plugins/Editor/BlueprintHeaderView/Source/BlueprintHeaderView/Public/BlueprintHeaderView.h
@@ -4,6 +4,9 @@
 
 #include "CoreMinimal.h"
 #include "Modules/ModuleManager.h"
+#include "AssetRegistry/AssetData.h"
+
+class FExtender;
 
 class FBlueprintHeaderViewModule : public IModuleInterface
 {
@@ -12,4 +15,17 @@ public:
 	/** IModuleInterface implementation */
 	virtual void StartupModule() override;
 	virtual void ShutdownModule() override;
+
+	/** Returns whether the Header View supports the given class */
+	static bool IsClassHeaderViewSupported(const UClass* InClass);
+
+	static void OpenHeaderViewForAsset(FAssetData InAssetData);
+private:
+	void SetupContentBrowserContextMenuExtender();
+
+	static TSharedRef<FExtender> OnExtendContentBrowserAssetSelectionMenu(const TArray<FAssetData>& SelectedAssets);
+
+private: 
+	/** Handle to our delegate so we can remove it at module shutdown */
+	FDelegateHandle ContentBrowserExtenderDelegateHandle;
 };
