Fix Incorrect Include path

#rnx

diff --git a/Engine/Plugins/Editor/BlueprintHeaderView/Source/BlueprintHeaderView/BlueprintHeaderView.Build.cs b/Engine/Plugins/Editor/BlueprintHeaderView/Source/BlueprintHeaderView/BlueprintHeaderView.Build.cs
--- a/Engine/Plugins/Editor/BlueprintHeaderView/Source/BlueprintHeaderView/BlueprintHeaderView.Build.cs
+++ b/Engine/Plugins/Editor/BlueprintHeaderView/Source/BlueprintHeaderView/BlueprintHeaderView.Build.cs
@@ -29,6 +29,7 @@ public class BlueprintHeaderView : ModuleRules
 				"BlueprintGraph",
 				"ApplicationCore",
 				"DeveloperSettings",
+				"WorkspaceMenuStructure",
 			}
 			);
 		
diff --git a/Engine/Plugins/Editor/BlueprintHeaderView/Source/BlueprintHeaderView/Private/BlueprintHeaderView.cpp b/Engine/Plugins/Editor/BlueprintHeaderView/Source/BlueprintHeaderView/Private/BlueprintHeaderView.cpp
--- a/Engine/Plugins/Editor/BlueprintHeaderView/Source/BlueprintHeaderView/Private/BlueprintHeaderView.cpp
+++ b/Engine/Plugins/Editor/BlueprintHeaderView/Source/BlueprintHeaderView/Private/BlueprintHeaderView.cpp
@@ -8,8 +8,8 @@
 #include "Textures/SlateIcon.h"
 #include "Framework/Docking/TabManager.h"
 #include "EditorStyleSet.h"
-#include "Editor/WorkspaceMenuStructure/Public/WorkspaceMenuStructure.h"
-#include "Editor/WorkspaceMenuStructure/Public/WorkspaceMenuStructureModule.h"
+#include "WorkspaceMenuStructure.h"
+#include "WorkspaceMenuStructureModule.h"
 #include "Widgets/Docking/SDockTab.h"
 #include "SBlueprintHeaderView.h"
 #include "ContentBrowserModule.h"
