BPHeaderView: Automatically select blueprint if Header View was summoned from the Blueprint Editor

diff --git a/Engine/Plugins/Editor/BlueprintHeaderView/Source/BlueprintHeaderView/BlueprintHeaderView.Build.cs b/Engine/Plugins/Editor/BlueprintHeaderView/Source/BlueprintHeaderView/BlueprintHeaderView.Build.cs
--- a/Engine/Plugins/Editor/BlueprintHeaderView/Source/BlueprintHeaderView/BlueprintHeaderView.Build.cs
+++ b/Engine/Plugins/Editor/BlueprintHeaderView/Source/BlueprintHeaderView/BlueprintHeaderView.Build.cs
@@ -30,6 +30,7 @@ public class BlueprintHeaderView : ModuleRules
 				"ApplicationCore",
 				"DeveloperSettings",
 				"WorkspaceMenuStructure",
+				"ToolMenus",
 			}
 			);
 		
diff --git a/Engine/Plugins/Editor/BlueprintHeaderView/Source/BlueprintHeaderView/Private/BlueprintHeaderView.cpp b/Engine/Plugins/Editor/BlueprintHeaderView/Source/BlueprintHeaderView/Private/BlueprintHeaderView.cpp
--- a/Engine/Plugins/Editor/BlueprintHeaderView/Source/BlueprintHeaderView/Private/BlueprintHeaderView.cpp
+++ b/Engine/Plugins/Editor/BlueprintHeaderView/Source/BlueprintHeaderView/Private/BlueprintHeaderView.cpp
@@ -16,6 +16,9 @@
 #include "Framework/Commands/UIAction.h"
 #include "Engine/Blueprint.h"
 #include "Framework/MultiBox/MultiBoxBuilder.h"
+#include "ToolMenus.h"
+#include "Toolkits/AssetEditorToolkit.h"
+#include "Toolkits/AssetEditorToolkitMenuContext.h"
 
 #define LOCTEXT_NAMESPACE "BlueprintHeaderViewApp"
 
@@ -41,6 +44,7 @@ void FBlueprintHeaderViewModule::StartupModule()
 		.SetGroup(WorkspaceMenu::GetMenuStructure().GetToolsCategory())
 		.SetIcon(FSlateIcon(FEditorStyle::GetStyleSetName(), "ClassIcon.Class"));
 
+	SetupAssetEditorMenuExtender();
 	SetupContentBrowserContextMenuExtender();
 }
 
@@ -69,13 +73,42 @@ void FBlueprintHeaderViewModule::OpenHeaderViewForAsset(FAssetData InAssetData)
 	if (HeaderViewTab.IsValid())
 	{
 		TSharedRef<SWidget> HeaderViewContentWidget = HeaderViewTab->GetContent();
-		if (HeaderViewContentWidget->GetWidgetClass().GetWidgetType() == SBlueprintHeaderView::StaticWidgetClass().GetWidgetType())
-		{
-			StaticCastSharedRef<SBlueprintHeaderView>(HeaderViewContentWidget)->OnAssetSelected(InAssetData);
-		}
+		StaticCastSharedRef<SBlueprintHeaderView>(HeaderViewContentWidget)->OnAssetSelected(InAssetData);
 	}
 }
 
+void FBlueprintHeaderViewModule::SetupAssetEditorMenuExtender()
+{
+	UToolMenu* Menu = UToolMenus::Get()->ExtendMenu("MainFrame.MainMenu.Asset");
+	FToolMenuSection& Section = Menu->FindOrAddSection("HeaderViewActions");
+	FToolMenuEntry& Entry = Section.AddDynamicEntry("BlueprintHeaderViewCommands", FNewToolMenuSectionDelegate::CreateLambda([](FToolMenuSection& InSection)
+		{
+			UAssetEditorToolkitMenuContext* MenuContext = InSection.FindContext<UAssetEditorToolkitMenuContext>();
+			if (MenuContext && MenuContext->Toolkit.IsValid())
+			{
+				TSharedPtr<FAssetEditorToolkit> Toolkit = MenuContext->Toolkit.Pin();
+				if (Toolkit->IsActuallyAnAsset() && Toolkit->GetObjectsCurrentlyBeingEdited()->Num() == 1)
+				{
+					for (const UObject* EditedObject : *Toolkit->GetObjectsCurrentlyBeingEdited())
+					{
+						if (IsClassHeaderViewSupported(EditedObject->GetClass()))
+						{
+							FAssetData BlueprintAssetData(EditedObject);
+
+							InSection.AddMenuEntry(
+								FName("OpenHeaderView"),
+								LOCTEXT("OpenAssetHeaderView", "Display in Blueprint Header View"),
+								LOCTEXT("OpenAssetHeaderViewTooltip", "Opens this Blueprint in the Blueprint Header View"),
+								FSlateIcon(FEditorStyle::GetStyleSetName(), "ClassIcon.Class"),
+								FUIAction(FExecuteAction::CreateStatic(&FBlueprintHeaderViewModule::OpenHeaderViewForAsset, BlueprintAssetData))
+							);
+						}
+					}
+				}
+			}
+		}));
+}
+
 void FBlueprintHeaderViewModule::SetupContentBrowserContextMenuExtender()
 {
 	FContentBrowserModule& ContentBrowserModule = FModuleManager::LoadModuleChecked<FContentBrowserModule>("ContentBrowser");
diff --git a/Engine/Plugins/Editor/BlueprintHeaderView/Source/BlueprintHeaderView/Public/BlueprintHeaderView.h b/Engine/Plugins/Editor/BlueprintHeaderView/Source/BlueprintHeaderView/Public/BlueprintHeaderView.h
--- a/Engine/Plugins/Editor/BlueprintHeaderView/Source/BlueprintHeaderView/Public/BlueprintHeaderView.h
+++ b/Engine/Plugins/Editor/BlueprintHeaderView/Source/BlueprintHeaderView/Public/BlueprintHeaderView.h
@@ -21,6 +21,8 @@ public:
 
 	static void OpenHeaderViewForAsset(FAssetData InAssetData);
 private:
+	void SetupAssetEditorMenuExtender();
+
 	void SetupContentBrowserContextMenuExtender();
 
 	static TSharedRef<FExtender> OnExtendContentBrowserAssetSelectionMenu(const TArray<FAssetData>& SelectedAssets);
