Fix Incorrect Property Flags in Header View

diff --git a/Engine/Plugins/Editor/BlueprintHeaderView/Source/BlueprintHeaderView/Private/HeaderViewVariableListItem.cpp b/Engine/Plugins/Editor/BlueprintHeaderView/Source/BlueprintHeaderView/Private/HeaderViewVariableListItem.cpp
--- a/Engine/Plugins/Editor/BlueprintHeaderView/Source/BlueprintHeaderView/Private/HeaderViewVariableListItem.cpp
+++ b/Engine/Plugins/Editor/BlueprintHeaderView/Source/BlueprintHeaderView/Private/HeaderViewVariableListItem.cpp
@@ -197,17 +197,13 @@ FString FHeaderViewVariableListItem::GetConditionalUPropertySpecifiers(const FPr
 		{
 			PropertySpecifiers.Emplace(TEXT("BlueprintAssignable"));
 		}
-		else if (VarProperty.HasAnyPropertyFlags(CPF_BlueprintReadOnly))
-		{
-			PropertySpecifiers.Emplace(TEXT("BlueprintReadOnly"));
-		}
 		else if (VarProperty.HasAnyPropertyFlags(CPF_Edit))
 		{
 			PropertySpecifiers.Emplace(TEXT("BlueprintReadWrite"));
 		}
 		else
 		{
-			PropertySpecifiers.Emplace(TEXT("BlueprintVisible"));
+			PropertySpecifiers.Emplace(TEXT("BlueprintReadOnly"));
 		}
 	}
 
