BPHeaderView: Add a Jump to Definition option to Function items via context menu or double click

diff --git a/Engine/Plugins/Editor/BlueprintHeaderView/Source/BlueprintHeaderView/Private/HeaderViewFunctionListItem.cpp b/Engine/Plugins/Editor/BlueprintHeaderView/Source/BlueprintHeaderView/Private/HeaderViewFunctionListItem.cpp
--- a/Engine/Plugins/Editor/BlueprintHeaderView/Source/BlueprintHeaderView/Private/HeaderViewFunctionListItem.cpp
+++ b/Engine/Plugins/Editor/BlueprintHeaderView/Source/BlueprintHeaderView/Private/HeaderViewFunctionListItem.cpp
@@ -7,6 +7,7 @@
 #include "K2Node_FunctionEntry.h"
 #include "Misc/EngineVersion.h"
 #include "Framework/MultiBox/MultiBoxBuilder.h"
+#include "Kismet2/KismetEditorUtilities.h"
 
 #define LOCTEXT_NAMESPACE "FHeaderViewFunctionListItem"
 
@@ -17,6 +18,12 @@ FHeaderViewListItemPtr FHeaderViewFunctionListItem::Create(const UK2Node_Functio
 
 void FHeaderViewFunctionListItem::ExtendContextMenu(FMenuBuilder& InMenuBuilder, TWeakObjectPtr<UBlueprint> InBlueprint)
 {
+	InMenuBuilder.AddMenuEntry(LOCTEXT("JumpToDefinition", "Jump to Definition"),
+		LOCTEXT("JumpToDefinitionTooltip", "Opens this function in the Blueprint Editor"),
+		FSlateIcon(),
+		FUIAction(FExecuteAction::CreateSP(this, &FHeaderViewFunctionListItem::JumpToDefinition, InBlueprint))
+	);
+
 	if (!IllegalName.IsNone())
 	{
 		InMenuBuilder.AddEditableText(LOCTEXT("RenameItem", "Rename Function"),
@@ -38,6 +45,11 @@ void FHeaderViewFunctionListItem::ExtendContextMenu(FMenuBuilder& InMenuBuilder,
 	}
 }
 
+void FHeaderViewFunctionListItem::OnMouseButtonDoubleClick(TWeakObjectPtr<UBlueprint> Blueprint)
+{
+	JumpToDefinition(Blueprint);
+}
+
 FString FHeaderViewFunctionListItem::GetConditionalUFunctionSpecifiers(const UFunction* SigFunction) const
 {
 	check(SigFunction);
@@ -342,4 +354,25 @@ void FHeaderViewFunctionListItem::OnRenameParameterTextCommitted(const FText& Co
 	}
 }
 
+void FHeaderViewFunctionListItem::JumpToDefinition(TWeakObjectPtr<UBlueprint> WeakBlueprint) const
+{
+	if (const UBlueprint* Blueprint = WeakBlueprint.Get())
+	{
+		const TObjectPtr<UEdGraph>* FunctionGraph = Blueprint->FunctionGraphs.FindByPredicate([GraphName=GraphName](const TObjectPtr<UEdGraph>& Graph)
+			{
+				return Graph->GetFName() == GraphName;
+			});
+
+		if (FunctionGraph)
+		{
+			TArray<UK2Node_FunctionEntry*> Entry;
+			FunctionGraph->Get()->GetNodesOfClass<UK2Node_FunctionEntry>(Entry);
+			if (ensureMsgf(Entry.Num() == 1 && Entry[0], TEXT("Function Graph \"%s\" in Blueprint \"%s\" does not have exactly one Entry Node"), *FunctionGraph->Get()->GetName(), *Blueprint->GetName()))
+			{
+				FKismetEditorUtilities::BringKismetToFocusAttentionOnObject(Entry[0]);
+			}
+		}
+	}
+}
+
 #undef LOCTEXT_NAMESPACE
diff --git a/Engine/Plugins/Editor/BlueprintHeaderView/Source/BlueprintHeaderView/Private/HeaderViewFunctionListItem.h b/Engine/Plugins/Editor/BlueprintHeaderView/Source/BlueprintHeaderView/Private/HeaderViewFunctionListItem.h
--- a/Engine/Plugins/Editor/BlueprintHeaderView/Source/BlueprintHeaderView/Private/HeaderViewFunctionListItem.h
+++ b/Engine/Plugins/Editor/BlueprintHeaderView/Source/BlueprintHeaderView/Private/HeaderViewFunctionListItem.h
@@ -15,6 +15,7 @@ struct FHeaderViewFunctionListItem : public FHeaderViewListItem
 
 	//~ FHeaderViewListItem Interface
 	virtual void ExtendContextMenu(FMenuBuilder& InMenuBuilder, TWeakObjectPtr<UBlueprint> InBlueprint) override;
+	virtual void OnMouseButtonDoubleClick(TWeakObjectPtr<UBlueprint> Blueprint) override;
 	//~ End FHeaderViewListItem Interface
 
 protected:
@@ -26,9 +27,10 @@ protected:
 	/** Adds Function parameters to the RichText and PlainText strings */
 	void AppendFunctionParameters(const UFunction* SignatureFunction);
 
+	// UI Callback Functions
 	void OnRenameFunctionTextCommitted(const FText& CommittedText, ETextCommit::Type TextCommitType, TWeakObjectPtr<UBlueprint> WeakBlueprint, FName OldGraphName);
 	void OnRenameParameterTextCommitted(const FText& CommittedText, ETextCommit::Type TextCommitType, TWeakObjectPtr<UBlueprint> WeakBlueprint, FName OldGraphName, FName OldParamName);
-
+	void JumpToDefinition(TWeakObjectPtr<UBlueprint> WeakBlueprint) const;
 protected:
 	/** None if the function name is legal C++, else the name of the function */
 	FName IllegalName = NAME_None;
diff --git a/Engine/Plugins/Editor/BlueprintHeaderView/Source/BlueprintHeaderView/Private/SBlueprintHeaderView.cpp b/Engine/Plugins/Editor/BlueprintHeaderView/Source/BlueprintHeaderView/Private/SBlueprintHeaderView.cpp
--- a/Engine/Plugins/Editor/BlueprintHeaderView/Source/BlueprintHeaderView/Private/SBlueprintHeaderView.cpp
+++ b/Engine/Plugins/Editor/BlueprintHeaderView/Source/BlueprintHeaderView/Private/SBlueprintHeaderView.cpp
@@ -287,6 +287,7 @@ void SBlueprintHeaderView::Construct(const FArguments& InArgs)
 				.ListItemsSource(&ListItems)
 				.OnGenerateRow(this, &SBlueprintHeaderView::GenerateRowForItem)
 				.OnContextMenuOpening(this, &SBlueprintHeaderView::OnContextMenuOpening)
+				.OnMouseButtonDoubleClick(this, &SBlueprintHeaderView::OnItemDoubleClicked)
 //			]
 		]
 	];
@@ -564,6 +565,11 @@ TSharedPtr<SWidget> SBlueprintHeaderView::OnContextMenuOpening()
 	return MenuBuilder.MakeWidget();
 }
 
+void SBlueprintHeaderView::OnItemDoubleClicked(FHeaderViewListItemPtr Item)
+{
+	Item->OnMouseButtonDoubleClick(SelectedBlueprint);
+}
+
 void SBlueprintHeaderView::OnBlueprintChanged(UBlueprint* InBlueprint)
 {
 	RepopulateListView();
diff --git a/Engine/Plugins/Editor/BlueprintHeaderView/Source/BlueprintHeaderView/Private/SBlueprintHeaderView.h b/Engine/Plugins/Editor/BlueprintHeaderView/Source/BlueprintHeaderView/Private/SBlueprintHeaderView.h
--- a/Engine/Plugins/Editor/BlueprintHeaderView/Source/BlueprintHeaderView/Private/SBlueprintHeaderView.h
+++ b/Engine/Plugins/Editor/BlueprintHeaderView/Source/BlueprintHeaderView/Private/SBlueprintHeaderView.h
@@ -46,6 +46,9 @@ struct FHeaderViewListItem : public TSharedFromThis<FHeaderViewListItem>
 	/** Allows the item to add items to the context menu if it is the only item selected */
 	virtual void ExtendContextMenu(FMenuBuilder& InMenuBuilder, TWeakObjectPtr<UBlueprint> InBlueprint) {}
 
+	/** Called when this List Item is double clicked */
+	virtual void OnMouseButtonDoubleClick(TWeakObjectPtr<UBlueprint> Blueprint) {};
+
 protected:
 	/** Empty base constructor hidden from public */
 	FHeaderViewListItem() {};
@@ -133,6 +136,9 @@ private:
 	/** Creates a context menu for the list view */
 	TSharedPtr<SWidget> OnContextMenuOpening();
 
+	/** Called when a List Item is double clicked */
+	void OnItemDoubleClicked(FHeaderViewListItemPtr Item);
+
 	/** Callback for when the selected blueprint is modified */
 	void OnBlueprintChanged(UBlueprint* InBlueprint);
 
