BPHeaderView: Expose FontSize and ListItem SelectionColor to config

diff --git a/Engine/Config/BaseEditorPerProjectUserSettings.ini b/Engine/Config/BaseEditorPerProjectUserSettings.ini
--- a/Engine/Config/BaseEditorPerProjectUserSettings.ini
+++ b/Engine/Config/BaseEditorPerProjectUserSettings.ini
@@ -930,5 +930,7 @@ bEnableReinstancing=True
 bAutomaticallyCompileNewClasses=True
 
 [/Script/BlueprintHeaderView.BlueprintHeaderViewSettings]
-SyntaxColors=(Comment=(R=0.300000,G=0.700000,B=0.100000,A=1.000000),Macro=(R=0.600000,G=0.200000,B=0.800000,A=1.000000),Typename=(R=0.000000,G=0.300000,B=0.300000,A=1.000000),Identifier=(R=1.000000,G=1.000000,B=1.000000,A=1.000000),Keyword=(R=0.000000,G=0.400000,B=0.800000,A=1.000000))
+SyntaxColors=(Comment=(R=0.258183,G=0.539479,B=0.068478,A=1.000000),Error=(R=0.863157,G=0.035601,B=0.035601,A=1.000000),Macro=(R=0.356400,G=0.040915,B=0.520996,A=1.000000),Typename=(R=0.000000,G=0.300000,B=0.300000,A=1.000000),Identifier=(R=1.000000,G=1.000000,B=1.000000,A=1.000000),Keyword=(R=0.019382,G=0.496933,B=1.000000,A=1.000000))
+FontSize=9
+SelectionColor=(R=0.051269,G=0.095307,B=0.158961,A=1.000000)
 SortMethod=None
diff --git a/Engine/Plugins/Editor/BlueprintHeaderView/Source/BlueprintHeaderView/Private/BlueprintHeaderView.cpp b/Engine/Plugins/Editor/BlueprintHeaderView/Source/BlueprintHeaderView/Private/BlueprintHeaderView.cpp
--- a/Engine/Plugins/Editor/BlueprintHeaderView/Source/BlueprintHeaderView/Private/BlueprintHeaderView.cpp
+++ b/Engine/Plugins/Editor/BlueprintHeaderView/Source/BlueprintHeaderView/Private/BlueprintHeaderView.cpp
@@ -19,6 +19,7 @@
 #include "ToolMenus.h"
 #include "Toolkits/AssetEditorToolkit.h"
 #include "Toolkits/AssetEditorToolkitMenuContext.h"
+#include "BlueprintHeaderViewSettings.h"
 
 #define LOCTEXT_NAMESPACE "BlueprintHeaderViewApp"
 
@@ -36,6 +37,9 @@ namespace BlueprintHeaderViewModule
 	}
 }
 
+FTextBlockStyle FBlueprintHeaderViewModule::HeaderViewTextStyle;
+FTableRowStyle FBlueprintHeaderViewModule::HeaderViewTableRowStyle;
+
 void FBlueprintHeaderViewModule::StartupModule()
 {
 	FGlobalTabmanager::Get()->RegisterNomadTabSpawner(BlueprintHeaderViewModule::HeaderViewTabName, FOnSpawnTab::CreateStatic(&BlueprintHeaderViewModule::CreateHeaderViewTab))
@@ -44,6 +48,15 @@ void FBlueprintHeaderViewModule::StartupModule()
 		.SetGroup(WorkspaceMenu::GetMenuStructure().GetToolsCategory())
 		.SetIcon(FSlateIcon(FEditorStyle::GetStyleSetName(), "ClassIcon.Class"));
 
+	HeaderViewTextStyle = FTextBlockStyle()
+		.SetFont(FAppStyle::Get().GetFontStyle("Mono"))
+		.SetFontSize(GetDefault<UBlueprintHeaderViewSettings>()->FontSize)
+		.SetColorAndOpacity(FSlateColor::UseForeground());
+
+	HeaderViewTableRowStyle = FCoreStyle::Get().GetWidgetStyle<FTableRowStyle>("TableView.Row");
+	HeaderViewTableRowStyle.ActiveBrush.TintColor = GetDefault<UBlueprintHeaderViewSettings>()->SelectionColor;
+	HeaderViewTableRowStyle.ActiveHoveredBrush.TintColor = GetDefault<UBlueprintHeaderViewSettings>()->SelectionColor;
+
 	SetupAssetEditorMenuExtender();
 	SetupContentBrowserContextMenuExtender();
 }
diff --git a/Engine/Plugins/Editor/BlueprintHeaderView/Source/BlueprintHeaderView/Private/BlueprintHeaderViewSettings.cpp b/Engine/Plugins/Editor/BlueprintHeaderView/Source/BlueprintHeaderView/Private/BlueprintHeaderViewSettings.cpp
--- a/Engine/Plugins/Editor/BlueprintHeaderView/Source/BlueprintHeaderView/Private/BlueprintHeaderViewSettings.cpp
+++ b/Engine/Plugins/Editor/BlueprintHeaderView/Source/BlueprintHeaderView/Private/BlueprintHeaderViewSettings.cpp
@@ -1,6 +1,19 @@
 // Copyright Epic Games, Inc. All Rights Reserved.
 
 #include "BlueprintHeaderViewSettings.h"
+#include "BlueprintHeaderView.h"
+#include "Styling/StyleColors.h"
+
+FHeaderViewSyntaxColors::FHeaderViewSyntaxColors()
+	: Comment(FStyleColors::AccentGreen.GetSpecifiedColor())
+	, Error(FStyleColors::Error.GetSpecifiedColor())
+	, Macro(FStyleColors::AccentPurple.GetSpecifiedColor())
+	, Typename(0.0f, 0.3f, 0.3f)
+	, Identifier(FStyleColors::White.GetSpecifiedColor())
+	, Keyword(FStyleColors::AccentBlue.GetSpecifiedColor())
+{
+
+}
 
 UBlueprintHeaderViewSettings::UBlueprintHeaderViewSettings()
 {
@@ -21,3 +34,23 @@ FName UBlueprintHeaderViewSettings::GetSectionName() const
 {
 	return TEXT("Blueprint Header View");
 }
+
+void UBlueprintHeaderViewSettings::PostEditChangeProperty(FPropertyChangedEvent& PropertyChangedEvent)
+{
+	Super::PostEditChangeProperty(PropertyChangedEvent);
+
+	FName PropertyChangedName = PropertyChangedEvent.GetPropertyName();
+
+	SyntaxColors = FHeaderViewSyntaxColors();
+
+	if (PropertyChangedName == GET_MEMBER_NAME_CHECKED(UBlueprintHeaderViewSettings, FontSize))
+	{
+		FBlueprintHeaderViewModule::HeaderViewTextStyle.SetFontSize(FontSize);
+	}
+
+	if (PropertyChangedName == GET_MEMBER_NAME_CHECKED(UBlueprintHeaderViewSettings, SelectionColor))
+	{
+		FBlueprintHeaderViewModule::HeaderViewTableRowStyle.ActiveBrush.TintColor = SelectionColor;
+		FBlueprintHeaderViewModule::HeaderViewTableRowStyle.ActiveHoveredBrush.TintColor = SelectionColor;
+	}
+}
diff --git a/Engine/Plugins/Editor/BlueprintHeaderView/Source/BlueprintHeaderView/Private/SBlueprintHeaderView.cpp b/Engine/Plugins/Editor/BlueprintHeaderView/Source/BlueprintHeaderView/Private/SBlueprintHeaderView.cpp
--- a/Engine/Plugins/Editor/BlueprintHeaderView/Source/BlueprintHeaderView/Private/SBlueprintHeaderView.cpp
+++ b/Engine/Plugins/Editor/BlueprintHeaderView/Source/BlueprintHeaderView/Private/SBlueprintHeaderView.cpp
@@ -1,6 +1,7 @@
 // Copyright Epic Games, Inc. All Rights Reserved.
 
 #include "SBlueprintHeaderView.h"
+#include "BlueprintHeaderView.h"
 #include "Widgets/Layout/SSpacer.h"
 #include "Widgets/Layout/SBox.h"
 //#include "Widgets/Layout/SScrollBox.h"
@@ -68,6 +69,9 @@ namespace
 
 			ModelText->Append(OriginalText.Mid(RunParseResult.ContentRange.BeginIndex, RunParseResult.ContentRange.Len()));
 
+			FTextBlockStyle TextStyle = FBlueprintHeaderViewModule::HeaderViewTextStyle;
+			TextStyle.SetColorAndOpacity(TextColor);
+
 			return FSlateTextRun::Create(RunInfo, ModelText, TextStyle);
 		}
 
@@ -75,9 +79,8 @@ namespace
 
 		FHeaderViewSyntaxDecorator(FString&& InName, const FSlateColor& InColor)
 			: DecoratorName(MoveTemp(InName))
-			, TextStyle(FEditorStyle::Get().GetWidgetStyle<FTextBlockStyle>("Log.Normal"))
+			, TextColor(InColor)
 		{
-			TextStyle.SetColorAndOpacity(InColor);
 		}
 
 	private:
@@ -85,8 +88,8 @@ namespace
 		/** Name of this decorator */
 		FString DecoratorName;
 
-		/** Style of this decorator */
-		FTextBlockStyle TextStyle;
+		/** Color of this decorator */
+		FSlateColor TextColor;
 
 	};
 
@@ -105,7 +108,7 @@ TSharedRef<SWidget> FHeaderViewListItem::GenerateWidgetForItem()
 		[
 			SNew(SRichTextBlock)
 			.Text(FText::FromString(RichTextString))
-			.TextStyle(FEditorStyle::Get(), "Log.Normal")
+			.TextStyle(&FBlueprintHeaderViewModule::HeaderViewTextStyle)
 			+SRichTextBlock::Decorator(FHeaderViewSyntaxDecorator::Create(HeaderViewSyntaxDecorators::CommentDecorator, SyntaxColors.Comment))
 			+SRichTextBlock::Decorator(FHeaderViewSyntaxDecorator::Create(HeaderViewSyntaxDecorators::ErrorDecorator, SyntaxColors.Error))
 			+SRichTextBlock::Decorator(FHeaderViewSyntaxDecorator::Create(HeaderViewSyntaxDecorators::IdentifierDecorator, SyntaxColors.Identifier))
@@ -308,6 +311,9 @@ void SBlueprintHeaderView::NotifyPostChange(const FPropertyChangedEvent& Propert
 	UBlueprintHeaderViewSettings* BlueprintHeaderViewSettings = GetMutableDefault<UBlueprintHeaderViewSettings>();
 	check(BlueprintHeaderViewSettings);
 	BlueprintHeaderViewSettings->SaveConfig();
+
+	// repopulate the list view to update text style/sorting method based on settings
+	RepopulateListView();
 }
 
 FText SBlueprintHeaderView::GetClassPickerText() const
@@ -373,6 +379,7 @@ void SBlueprintHeaderView::OnAssetSelected(const FAssetData& SelectedAsset)
 TSharedRef<ITableRow> SBlueprintHeaderView::GenerateRowForItem(FHeaderViewListItemPtr Item, const TSharedRef<STableViewBase>& OwnerTable) const
 {
 	return SNew(STableRow<FHeaderViewListItemPtr>, OwnerTable)
+		.Style(&FBlueprintHeaderViewModule::HeaderViewTableRowStyle)
 		.Content()
 		[
 			Item->GenerateWidgetForItem()
diff --git a/Engine/Plugins/Editor/BlueprintHeaderView/Source/BlueprintHeaderView/Public/BlueprintHeaderView.h b/Engine/Plugins/Editor/BlueprintHeaderView/Source/BlueprintHeaderView/Public/BlueprintHeaderView.h
--- a/Engine/Plugins/Editor/BlueprintHeaderView/Source/BlueprintHeaderView/Public/BlueprintHeaderView.h
+++ b/Engine/Plugins/Editor/BlueprintHeaderView/Source/BlueprintHeaderView/Public/BlueprintHeaderView.h
@@ -5,6 +5,7 @@
 #include "CoreMinimal.h"
 #include "Modules/ModuleManager.h"
 #include "AssetRegistry/AssetData.h"
+#include "Styling/SlateTypes.h"
 
 class FExtender;
 
@@ -27,6 +28,13 @@ private:
 
 	static TSharedRef<FExtender> OnExtendContentBrowserAssetSelectionMenu(const TArray<FAssetData>& SelectedAssets);
 
+public:
+	/** Text style for the Header View Output */
+	static FTextBlockStyle HeaderViewTextStyle;
+
+	/** TableRow style for the Header View Output */
+	static FTableRowStyle HeaderViewTableRowStyle;
+
 private: 
 	/** Handle to our delegate so we can remove it at module shutdown */
 	FDelegateHandle ContentBrowserExtenderDelegateHandle;
diff --git a/Engine/Plugins/Editor/BlueprintHeaderView/Source/BlueprintHeaderView/Public/BlueprintHeaderViewSettings.h b/Engine/Plugins/Editor/BlueprintHeaderView/Source/BlueprintHeaderView/Public/BlueprintHeaderViewSettings.h
--- a/Engine/Plugins/Editor/BlueprintHeaderView/Source/BlueprintHeaderView/Public/BlueprintHeaderViewSettings.h
+++ b/Engine/Plugins/Editor/BlueprintHeaderView/Source/BlueprintHeaderView/Public/BlueprintHeaderViewSettings.h
@@ -16,25 +16,27 @@ USTRUCT(NotBlueprintable)
 struct FHeaderViewSyntaxColors
 {
 	GENERATED_BODY()
-	
+
+	FHeaderViewSyntaxColors();
+
 public:
 	UPROPERTY(config, EditAnywhere, Category="Settings")
-	FLinearColor Comment = FLinearColor(0.3f, 0.7f, 0.1f, 1.0f);
+	FLinearColor Comment;
 	
 	UPROPERTY(config, EditAnywhere, Category="Settings")
-	FLinearColor Error = FLinearColor(1.0f, 0.0f, 0.0f, 1.0f);
+	FLinearColor Error;
 	
 	UPROPERTY(config, EditAnywhere, Category="Settings")
-	FLinearColor Macro = FLinearColor(0.6f, 0.2f, 0.8f, 1.0f);
+	FLinearColor Macro;
 	
 	UPROPERTY(config, EditAnywhere, Category="Settings")
-	FLinearColor Typename = FLinearColor::White;
+	FLinearColor Typename;
 	
 	UPROPERTY(config, EditAnywhere, Category="Settings")
-	FLinearColor Identifier = FLinearColor::White;
+	FLinearColor Identifier;
 	
 	UPROPERTY(config, EditAnywhere, Category="Settings")
-	FLinearColor Keyword = FLinearColor(0.0f, 0.4f, 0.8f, 1.0f);
+	FLinearColor Keyword;
 	
 };
 
@@ -53,12 +55,25 @@ public:
 	virtual FName GetSectionName() const override;
 	//~ End UDeveloperSettings interface
 
+protected:
+	//~ Begin UObject interface
+	virtual void PostEditChangeProperty(struct FPropertyChangedEvent& PropertyChangedEvent) override;
+	//~ End UObject interface
+
 public:
-	
+
 	/** Syntax Highlighting colors for Blueprint Header View output */
-	UPROPERTY(config, EditAnywhere, Category="Settings")
+	UPROPERTY(config, EditAnywhere, Category="Settings|Style")
 	FHeaderViewSyntaxColors SyntaxColors;
+	
+	/** Highlight color for selected items in the Blueprint Header View output */
+	UPROPERTY(config, EditAnywhere, Category="Settings|Style")
+	FLinearColor SelectionColor = FLinearColor(0.3f, 0.3f, 1.0f, 1.0f);
 
+	/** Font Size for the Blueprint Header View output */
+	UPROPERTY(config, EditAnywhere, Category="Settings|Style")
+	int32 FontSize = 9;
+	
 	/** Sorting Method for Header View Functions and Properties */
 	UPROPERTY(config, EditAnywhere, Category="Settings")
 	EHeaderViewSortMethod SortMethod = EHeaderViewSortMethod::None;
