Fix BP Header View using the wrong font

diff --git a/Engine/Plugins/Editor/BlueprintHeaderView/Source/BlueprintHeaderView/Private/BlueprintHeaderView.cpp b/Engine/Plugins/Editor/BlueprintHeaderView/Source/BlueprintHeaderView/Private/BlueprintHeaderView.cpp
--- a/Engine/Plugins/Editor/BlueprintHeaderView/Source/BlueprintHeaderView/Private/BlueprintHeaderView.cpp
+++ b/Engine/Plugins/Editor/BlueprintHeaderView/Source/BlueprintHeaderView/Private/BlueprintHeaderView.cpp
@@ -49,8 +49,7 @@ void FBlueprintHeaderViewModule::StartupModule()
 		.SetIcon(FSlateIcon(FEditorStyle::GetStyleSetName(), "ClassIcon.Class"));
 
 	HeaderViewTextStyle = FTextBlockStyle()
-		.SetFont(FAppStyle::Get().GetFontStyle("Mono"))
-		.SetFontSize(GetDefault<UBlueprintHeaderViewSettings>()->FontSize)
+		.SetFont(FCoreStyle::GetDefaultFontStyle("Mono", GetDefault<UBlueprintHeaderViewSettings>()->FontSize))
 		.SetColorAndOpacity(FSlateColor::UseForeground());
 
 	HeaderViewTableRowStyle = FCoreStyle::Get().GetWidgetStyle<FTableRowStyle>("TableView.Row");
