Align Name validation in the BP Header View with the Blueprint Editor

diff --git a/Engine/Plugins/Editor/BlueprintHeaderView/Source/BlueprintHeaderView/BlueprintHeaderView.Build.cs b/Engine/Plugins/Editor/BlueprintHeaderView/Source/BlueprintHeaderView/BlueprintHeaderView.Build.cs
--- a/Engine/Plugins/Editor/BlueprintHeaderView/Source/BlueprintHeaderView/BlueprintHeaderView.Build.cs
+++ b/Engine/Plugins/Editor/BlueprintHeaderView/Source/BlueprintHeaderView/BlueprintHeaderView.Build.cs
@@ -31,6 +31,7 @@ public class BlueprintHeaderView : ModuleRules
 				"DeveloperSettings",
 				"WorkspaceMenuStructure",
 				"ToolMenus",
+				"AssetTools",
 			}
 			);
 		
diff --git a/Engine/Plugins/Editor/BlueprintHeaderView/Source/BlueprintHeaderView/Private/HeaderViewClassListItem.cpp b/Engine/Plugins/Editor/BlueprintHeaderView/Source/BlueprintHeaderView/Private/HeaderViewClassListItem.cpp
--- a/Engine/Plugins/Editor/BlueprintHeaderView/Source/BlueprintHeaderView/Private/HeaderViewClassListItem.cpp
+++ b/Engine/Plugins/Editor/BlueprintHeaderView/Source/BlueprintHeaderView/Private/HeaderViewClassListItem.cpp
@@ -5,6 +5,8 @@
 #include "String/LineEndings.h"
 #include "Kismet2/BlueprintEditorUtils.h"
 #include "Framework/MultiBox/MultiBoxBuilder.h"
+#include "AssetViewUtils.h"
+#include "AssetToolsModule.h"
 
 #define LOCTEXT_NAMESPACE "FHeaderViewClassListItem"
 
@@ -19,10 +21,11 @@ void FHeaderViewClassListItem::ExtendContextMenu(FMenuBuilder& InMenuBuilder, TW
 	{
 		if (UBlueprint* Blueprint = InBlueprint.Get())
 		{
-			InMenuBuilder.AddEditableText(LOCTEXT("RenameBlueprint", "Rename Blueprint"),
+			InMenuBuilder.AddVerifiedEditableText(LOCTEXT("RenameBlueprint", "Rename Blueprint"),
 				LOCTEXT("RenameItemTooltip", "Renames this Blueprint\nThis Blueprint name is not a legal C++ identifier."),
 				FSlateIcon(),
 				FText::FromString(Blueprint->GetName()),
+				FOnVerifyTextChanged::CreateSP(this, &FHeaderViewClassListItem::OnVerifyRenameTextChanged, InBlueprint),
 				FOnTextCommitted::CreateSP(this, &FHeaderViewClassListItem::OnRenameTextComitted, InBlueprint)
 			);
 		}
@@ -152,16 +155,56 @@ FHeaderViewClassListItem::FHeaderViewClassListItem(TWeakObjectPtr<UBlueprint> In
 	}
 }
 
+FString FHeaderViewClassListItem::GetRenamedBlueprintPath(const UBlueprint* Blueprint, const FString& NewName) const
+{
+	check(Blueprint);
+	FString NewPath = Blueprint->GetPathName();
+	int32 Index = 0;
+	if (!ensure(NewPath.FindLastChar(TEXT('/'), Index)))
+	{
+		// AssetViewUtils::IsValidObjectPathForCreate will return false for an empty string
+		return TEXT("");
+	}
+	NewPath.LeftInline(Index);
+	NewPath /= FString::Printf(TEXT("%s.%s"), *NewName, *NewName);
+	return NewPath;
+}
+
+bool FHeaderViewClassListItem::OnVerifyRenameTextChanged(const FText& InNewName, FText& OutErrorText, TWeakObjectPtr<UBlueprint> InBlueprint)
+{
+	if (const UBlueprint* Blueprint = InBlueprint.Get())
+	{
+		const FString NewPath = GetRenamedBlueprintPath(Blueprint, InNewName.ToString());
+		if (!AssetViewUtils::IsValidObjectPathForCreate(NewPath, Blueprint->GetClass(), OutErrorText))
+		{
+			return false;
+		}
+
+		if (!IsValidCPPIdentifier(InNewName.ToString()))
+		{
+			OutErrorText = InvalidCPPIdentifierErrorText;
+			return false;
+		}
+
+		return true;
+	}
+
+	return false;
+}
 
 void FHeaderViewClassListItem::OnRenameTextComitted(const FText& CommittedText, ETextCommit::Type TextCommitType, TWeakObjectPtr<UBlueprint> InBlueprint)
 {
 	if (TextCommitType == ETextCommit::OnEnter)
 	{
-		if (UBlueprint* Blueprint = InBlueprint.Get())
+		FText ErrorText;
+		if (OnVerifyRenameTextChanged(CommittedText, ErrorText, InBlueprint))
 		{
-			if (IsValidCPPIdentifier(CommittedText.ToString()))
+			if (UBlueprint* Blueprint = InBlueprint.Get())
 			{
-				Blueprint->Rename(*CommittedText.ToString());
+				const FString NewPath = GetRenamedBlueprintPath(Blueprint, CommittedText.ToString());
+				FAssetToolsModule& AssetToolsModule = FModuleManager::LoadModuleChecked<FAssetToolsModule>("AssetTools");
+				TArray<FAssetRenameData> AssetToRename = { FAssetRenameData(Blueprint, NewPath) };
+				AssetToolsModule.Get().RenameAssets(AssetToRename);
 				FBlueprintEditorUtils::MarkBlueprintAsStructurallyModified(Blueprint);
 			}
 		}
diff --git a/Engine/Plugins/Editor/BlueprintHeaderView/Source/BlueprintHeaderView/Private/HeaderViewClassListItem.h b/Engine/Plugins/Editor/BlueprintHeaderView/Source/BlueprintHeaderView/Private/HeaderViewClassListItem.h
--- a/Engine/Plugins/Editor/BlueprintHeaderView/Source/BlueprintHeaderView/Private/HeaderViewClassListItem.h
+++ b/Engine/Plugins/Editor/BlueprintHeaderView/Source/BlueprintHeaderView/Private/HeaderViewClassListItem.h
@@ -22,6 +22,10 @@ protected:
 
 	FHeaderViewClassListItem(TWeakObjectPtr<UBlueprint> InBlueprint);
 
+	FString GetRenamedBlueprintPath(const UBlueprint* Blueprint, const FString& NewName) const;
+
+	bool OnVerifyRenameTextChanged(const FText& InNewName, FText& OutErrorText, TWeakObjectPtr<UBlueprint> InBlueprint);
+
 	void OnRenameTextComitted(const FText& CommittedText, ETextCommit::Type TextCommitType, TWeakObjectPtr<UBlueprint> InBlueprint);
 
 protected:
diff --git a/Engine/Plugins/Editor/BlueprintHeaderView/Source/BlueprintHeaderView/Private/HeaderViewFunctionListItem.cpp b/Engine/Plugins/Editor/BlueprintHeaderView/Source/BlueprintHeaderView/Private/HeaderViewFunctionListItem.cpp
--- a/Engine/Plugins/Editor/BlueprintHeaderView/Source/BlueprintHeaderView/Private/HeaderViewFunctionListItem.cpp
+++ b/Engine/Plugins/Editor/BlueprintHeaderView/Source/BlueprintHeaderView/Private/HeaderViewFunctionListItem.cpp
@@ -26,20 +26,22 @@ void FHeaderViewFunctionListItem::ExtendContextMenu(FMenuBuilder& InMenuBuilder,
 
 	if (!IllegalName.IsNone())
 	{
-		InMenuBuilder.AddEditableText(LOCTEXT("RenameItem", "Rename Function"),
+		InMenuBuilder.AddVerifiedEditableText(LOCTEXT("RenameItem", "Rename Function"),
 			LOCTEXT("RenameItemTooltip", "Renames this function in the Blueprint\nThis function name is not a legal C++ identifier."),
 			FSlateIcon(),
 			FText::FromName(IllegalName),
+			FOnVerifyTextChanged::CreateSP(this, &FHeaderViewFunctionListItem::OnVerifyRenameFunctionTextChanged, InBlueprint),
 			FOnTextCommitted::CreateSP(this, &FHeaderViewFunctionListItem::OnRenameFunctionTextCommitted, InBlueprint, GraphName)
 		);
 	}
 
 	for (FName IllegalParam : IllegalParameters)
 	{
-		InMenuBuilder.AddEditableText(LOCTEXT("RenameParm", "Rename Parameter"),
+		InMenuBuilder.AddVerifiedEditableText(LOCTEXT("RenameParm", "Rename Parameter"),
 			LOCTEXT("RenameParmTooltip", "Renames this function parameter in the Blueprint\nThis parameter name is not a legal C++ identifier."),
 			FSlateIcon(),
 			FText::FromName(IllegalParam),
+			FOnVerifyTextChanged::CreateSP(this, &FHeaderViewFunctionListItem::OnVerifyRenameParameterTextChanged, InBlueprint, GraphName),
 			FOnTextCommitted::CreateSP(this, &FHeaderViewFunctionListItem::OnRenameParameterTextCommitted, InBlueprint, GraphName, IllegalParam)
 		);
 	}
@@ -319,6 +321,67 @@ void FHeaderViewFunctionListItem::OnRenameFunctionTextCommitted(const FText& Com
 	}
 }
 
+bool FHeaderViewFunctionListItem::OnVerifyRenameFunctionTextChanged(const FText& InNewName, FText& OutErrorText, TWeakObjectPtr<UBlueprint> WeakBlueprint)
+{
+	if (const UBlueprint* Blueprint = WeakBlueprint.Get())
+	{
+		if (!IsValidCPPIdentifier(InNewName.ToString()))
+		{
+			OutErrorText = InvalidCPPIdentifierErrorText;
+			return false;
+		}
+
+		FKismetNameValidator NameValidator(Blueprint, NAME_None, nullptr);
+		const EValidatorResult Result = NameValidator.IsValid(InNewName.ToString());
+		if (Result != EValidatorResult::Ok)
+		{
+			OutErrorText = GetErrorTextFromValidatorResult(Result);
+			return false;
+		}
+
+		return true;
+	}
+
+	return false;
+}
+
+bool FHeaderViewFunctionListItem::OnVerifyRenameParameterTextChanged(const FText& InNewName, FText& OutErrorText, TWeakObjectPtr<UBlueprint> WeakBlueprint, FName OldGraphName)
+{
+	if (const UBlueprint* Blueprint = WeakBlueprint.Get())
+	{
+		if (!IsValidCPPIdentifier(InNewName.ToString()))
+		{
+			OutErrorText = InvalidCPPIdentifierErrorText;
+			return false;
+		}
+
+		const TObjectPtr<UEdGraph>* FunctionGraph = Blueprint->FunctionGraphs.FindByPredicate([OldGraphName](const TObjectPtr<UEdGraph>& Graph)
+			{
+				return Graph->GetFName() == OldGraphName;
+			});
+
+		if (FunctionGraph)
+		{
+			TArray<UK2Node_FunctionEntry*> Entry;
+			FunctionGraph->Get()->GetNodesOfClass<UK2Node_FunctionEntry>(Entry);
+			if (ensureMsgf(Entry.Num() == 1 && Entry[0], TEXT("Function Graph \"%s\" in Blueprint \"%s\" does not have exactly one Entry Node"), *FunctionGraph->Get()->GetName(), *Blueprint->GetName()))
+			{
+				FKismetNameValidator NameValidator(Blueprint, NAME_None, Entry[0]->FindSignatureFunction());
+				const EValidatorResult Result = NameValidator.IsValid(InNewName.ToString());
+				if (Result != EValidatorResult::Ok)
+				{
+					OutErrorText = GetErrorTextFromValidatorResult(Result);
+					return false;
+				}
+			}
+		}
+
+		return true;
+	}
+
+	return false;
+}
+
 void FHeaderViewFunctionListItem::OnRenameParameterTextCommitted(const FText& CommittedText, ETextCommit::Type TextCommitType, TWeakObjectPtr<UBlueprint> WeakBlueprint, FName OldGraphName, FName OldParamName)
 {
 	if (TextCommitType == ETextCommit::OnEnter)
diff --git a/Engine/Plugins/Editor/BlueprintHeaderView/Source/BlueprintHeaderView/Private/HeaderViewFunctionListItem.h b/Engine/Plugins/Editor/BlueprintHeaderView/Source/BlueprintHeaderView/Private/HeaderViewFunctionListItem.h
--- a/Engine/Plugins/Editor/BlueprintHeaderView/Source/BlueprintHeaderView/Private/HeaderViewFunctionListItem.h
+++ b/Engine/Plugins/Editor/BlueprintHeaderView/Source/BlueprintHeaderView/Private/HeaderViewFunctionListItem.h
@@ -28,6 +28,8 @@ protected:
 	void AppendFunctionParameters(const UFunction* SignatureFunction);
 
 	// UI Callback Functions
+	bool OnVerifyRenameFunctionTextChanged(const FText& InNewName, FText& OutErrorText, TWeakObjectPtr<UBlueprint> WeakBlueprint);
+	bool OnVerifyRenameParameterTextChanged(const FText& InNewName, FText& OutErrorText, TWeakObjectPtr<UBlueprint> WeakBlueprint, FName OldGraphName);
 	void OnRenameFunctionTextCommitted(const FText& CommittedText, ETextCommit::Type TextCommitType, TWeakObjectPtr<UBlueprint> WeakBlueprint, FName OldGraphName);
 	void OnRenameParameterTextCommitted(const FText& CommittedText, ETextCommit::Type TextCommitType, TWeakObjectPtr<UBlueprint> WeakBlueprint, FName OldGraphName, FName OldParamName);
 	void JumpToDefinition(TWeakObjectPtr<UBlueprint> WeakBlueprint) const;
diff --git a/Engine/Plugins/Editor/BlueprintHeaderView/Source/BlueprintHeaderView/Private/HeaderViewVariableListItem.cpp b/Engine/Plugins/Editor/BlueprintHeaderView/Source/BlueprintHeaderView/Private/HeaderViewVariableListItem.cpp
--- a/Engine/Plugins/Editor/BlueprintHeaderView/Source/BlueprintHeaderView/Private/HeaderViewVariableListItem.cpp
+++ b/Engine/Plugins/Editor/BlueprintHeaderView/Source/BlueprintHeaderView/Private/HeaderViewVariableListItem.cpp
@@ -19,10 +19,11 @@ void FHeaderViewVariableListItem::ExtendContextMenu(FMenuBuilder& InMenuBuilder,
 {
 	if (!IllegalName.IsNone())
 	{
-		InMenuBuilder.AddEditableText(LOCTEXT("RenameItem", "Rename Variable"),
+		InMenuBuilder.AddVerifiedEditableText(LOCTEXT("RenameItem", "Rename Variable"),
 			LOCTEXT("RenameItemTooltip", "Renames this variable in the Blueprint\nThis variable name is not a legal C++ identifier."),
 			FSlateIcon(),
 			FText::FromName(IllegalName),
+			FOnVerifyTextChanged::CreateSP(this, &FHeaderViewVariableListItem::OnVerifyRenameTextChanged, InBlueprint),
 			FOnTextCommitted::CreateSP(this, &FHeaderViewVariableListItem::OnRenameTextCommitted, InBlueprint)
 		);
 	}
@@ -328,6 +329,30 @@ FString FHeaderViewVariableListItem::GetOwningClassName(const FProperty& VarProp
 	return TEXT("");
 }
 
+bool FHeaderViewVariableListItem::OnVerifyRenameTextChanged(const FText& InNewName, FText& OutErrorText, TWeakObjectPtr<UBlueprint> WeakBlueprint)
+{
+	if (const UBlueprint* Blueprint = WeakBlueprint.Get())
+	{
+		if (!IsValidCPPIdentifier(InNewName.ToString()))
+		{
+			OutErrorText = InvalidCPPIdentifierErrorText;
+			return false;
+		}
+
+		FKismetNameValidator NameValidator(Blueprint, NAME_None, nullptr);
+		const EValidatorResult Result = NameValidator.IsValid(InNewName.ToString());
+		if (Result != EValidatorResult::Ok)
+		{
+			OutErrorText = GetErrorTextFromValidatorResult(Result);
+			return false;
+		}
+
+		return true;
+	}
+
+	return false;
+}
+
 void FHeaderViewVariableListItem::OnRenameTextCommitted(const FText& CommittedText, ETextCommit::Type TextCommitType, TWeakObjectPtr<UBlueprint> WeakBlueprint)
 {
 	if (TextCommitType == ETextCommit::OnEnter)
diff --git a/Engine/Plugins/Editor/BlueprintHeaderView/Source/BlueprintHeaderView/Private/HeaderViewVariableListItem.h b/Engine/Plugins/Editor/BlueprintHeaderView/Source/BlueprintHeaderView/Private/HeaderViewVariableListItem.h
--- a/Engine/Plugins/Editor/BlueprintHeaderView/Source/BlueprintHeaderView/Private/HeaderViewVariableListItem.h
+++ b/Engine/Plugins/Editor/BlueprintHeaderView/Source/BlueprintHeaderView/Private/HeaderViewVariableListItem.h
@@ -29,6 +29,7 @@ protected:
 	/** Returns the name of the owning class */
 	FString GetOwningClassName(const FProperty& VarProperty) const;
 
+	bool OnVerifyRenameTextChanged(const FText& InNewName, FText& OutErrorText, TWeakObjectPtr<UBlueprint> WeakBlueprint);
 	void OnRenameTextCommitted(const FText& CommittedText, ETextCommit::Type TextCommitType, TWeakObjectPtr<UBlueprint> WeakBlueprint);
 protected:
 
diff --git a/Engine/Plugins/Editor/BlueprintHeaderView/Source/BlueprintHeaderView/Private/SBlueprintHeaderView.cpp b/Engine/Plugins/Editor/BlueprintHeaderView/Source/BlueprintHeaderView/Private/SBlueprintHeaderView.cpp
--- a/Engine/Plugins/Editor/BlueprintHeaderView/Source/BlueprintHeaderView/Private/SBlueprintHeaderView.cpp
+++ b/Engine/Plugins/Editor/BlueprintHeaderView/Source/BlueprintHeaderView/Private/SBlueprintHeaderView.cpp
@@ -97,6 +97,8 @@ namespace
 
 // FHeaderViewListItem ////////////////////////////////////////////////////////
 
+const FText FHeaderViewListItem::InvalidCPPIdentifierErrorText = LOCTEXT("CPPIdentifierError", "Name is not a valid C++ Identifier");
+
 TSharedRef<SWidget> FHeaderViewListItem::GenerateWidgetForItem()
 {
 	const UBlueprintHeaderViewSettings* HeaderViewSettings = GetDefault<UBlueprintHeaderViewSettings>();
@@ -191,6 +193,24 @@ bool FHeaderViewListItem::IsValidCPPIdentifier(const FString& InIdentifier)
 	return FRegexMatcher(RegexPattern, InIdentifier).FindNext();
 }
 
+const FText& FHeaderViewListItem::GetErrorTextFromValidatorResult(EValidatorResult Result)
+{
+	static const TMap<EValidatorResult, FText> ErrorTextMap =
+	{
+		{EValidatorResult::AlreadyInUse, LOCTEXT("AlreadyInUse", "The name is already in use.")},
+		{EValidatorResult::EmptyName, LOCTEXT("EmptyName", "The name is empty.")},
+		{EValidatorResult::TooLong, LOCTEXT("TooLong", "The name is too long.")},
+		{EValidatorResult::LocallyInUse, LOCTEXT("LocallyInUse", "The name is already in use locally.")}
+	};
+
+	if (const FText* ErrorText = ErrorTextMap.Find(Result))
+	{
+		return *ErrorText;
+	}
+
+	return FText::GetEmpty();
+}
+
 // SBlueprintHeaderView ///////////////////////////////////////////////////////
 
 void SBlueprintHeaderView::Construct(const FArguments& InArgs)
diff --git a/Engine/Plugins/Editor/BlueprintHeaderView/Source/BlueprintHeaderView/Private/SBlueprintHeaderView.h b/Engine/Plugins/Editor/BlueprintHeaderView/Source/BlueprintHeaderView/Private/SBlueprintHeaderView.h
--- a/Engine/Plugins/Editor/BlueprintHeaderView/Source/BlueprintHeaderView/Private/SBlueprintHeaderView.h
+++ b/Engine/Plugins/Editor/BlueprintHeaderView/Source/BlueprintHeaderView/Private/SBlueprintHeaderView.h
@@ -8,6 +8,7 @@
 #include "Widgets/SCompoundWidget.h"
 #include "Widgets/Views/SListView.h"
 #include "Misc/NotifyHook.h"
+#include "Kismet2/Kismet2NameValidators.h"
 
 class SComboButton;
 class UBlueprint;
@@ -72,6 +73,9 @@ protected:
 	/** Checks whether a string is a valid C++ identifier */
 	static bool IsValidCPPIdentifier(const FString& InIdentifier);
 
+	/** Returns the correct Error text for a Validator Result */
+	static const FText& GetErrorTextFromValidatorResult(EValidatorResult Result);
+
 protected:
 	/** A rich text representation of the item, including syntax highlighting and errors */
 	FString RichTextString;
@@ -79,6 +83,8 @@ protected:
 	/** A raw string representation of the item, used for copying the item */
 	FString RawItemString;
 
+	/** Validation Error text for when a new name is not a valid C++ Identifier */
+	static const FText InvalidCPPIdentifierErrorText;
 };
 
 using FHeaderViewListItemPtr = TSharedPtr<FHeaderViewListItem>;
