Rename Blueprint Header View to clearer Blueprint Equivalent C++ Header Preview

diff --git a/Engine/Plugins/Editor/BlueprintHeaderView/BlueprintHeaderView.uplugin b/Engine/Plugins/Editor/BlueprintHeaderView/BlueprintHeaderView.uplugin
--- a/Engine/Plugins/Editor/BlueprintHeaderView/BlueprintHeaderView.uplugin
+++ b/Engine/Plugins/Editor/BlueprintHeaderView/BlueprintHeaderView.uplugin
@@ -2,7 +2,7 @@
 	"FileVersion": 3,
 	"Version": 1,
 	"VersionName": "1.0",
-	"FriendlyName": "Blueprint Header View",
+	"FriendlyName": "Blueprint C++ Header Preview",
 	"Description": "A tool to help convert Blueprint Classes to Native C++.",
 	"Category": "Blueprints",
 	"CreatedBy": "Epic Games, Inc.",
diff --git a/Engine/Plugins/Editor/BlueprintHeaderView/Source/BlueprintHeaderView/Private/BlueprintHeaderView.cpp b/Engine/Plugins/Editor/BlueprintHeaderView/Source/BlueprintHeaderView/Private/BlueprintHeaderView.cpp
--- a/Engine/Plugins/Editor/BlueprintHeaderView/Source/BlueprintHeaderView/Private/BlueprintHeaderView.cpp
+++ b/Engine/Plugins/Editor/BlueprintHeaderView/Source/BlueprintHeaderView/Private/BlueprintHeaderView.cpp
@@ -43,7 +43,7 @@ FTableRowStyle FBlueprintHeaderViewModule::HeaderViewTableRowStyle;
 void FBlueprintHeaderViewModule::StartupModule()
 {
 	FGlobalTabmanager::Get()->RegisterNomadTabSpawner(BlueprintHeaderViewModule::HeaderViewTabName, FOnSpawnTab::CreateStatic(&BlueprintHeaderViewModule::CreateHeaderViewTab))
-		.SetDisplayName(LOCTEXT("TabTitle", "Blueprint Header View"))
+		.SetDisplayName(LOCTEXT("TabTitle", "C++ Header Preview"))
 		.SetTooltipText(LOCTEXT("TooltipText", "Displays a Blueprint Class in C++ Header format."))
 		.SetGroup(WorkspaceMenu::GetMenuStructure().GetToolsCategory())
 		.SetIcon(FSlateIcon(FEditorStyle::GetStyleSetName(), "ClassIcon.Class"));
@@ -109,8 +109,8 @@ void FBlueprintHeaderViewModule::SetupAssetEditorMenuExtender()
 
 							InSection.AddMenuEntry(
 								FName("OpenHeaderView"),
-								LOCTEXT("OpenAssetHeaderView", "Display in Blueprint Header View"),
-								LOCTEXT("OpenAssetHeaderViewTooltip", "Opens this Blueprint in the Blueprint Header View"),
+								LOCTEXT("OpenAssetHeaderView", "Preview Equivalent C++ Header"),
+								LOCTEXT("OpenAssetHeaderViewTooltip", "Provides a preview of what this class could look like in C++"),
 								FSlateIcon(FEditorStyle::GetStyleSetName(), "ClassIcon.Class"),
 								FUIAction(FExecuteAction::CreateStatic(&FBlueprintHeaderViewModule::OpenHeaderViewForAsset, BlueprintAssetData))
 							);
@@ -140,8 +140,8 @@ TSharedRef<FExtender> FBlueprintHeaderViewModule::OnExtendContentBrowserAssetSel
 			Extender->AddMenuExtension("GetAssetActions", EExtensionHook::After, nullptr, FMenuExtensionDelegate::CreateLambda(
 				[SelectedAssets](FMenuBuilder& MenuBuilder) {
 					MenuBuilder.AddMenuEntry(
-						LOCTEXT("OpenHeaderView", "Display in Blueprint Header View"),
-						LOCTEXT("OpenHeaderViewTooltip", "Opens this Blueprint in the Blueprint Header View"),
+						LOCTEXT("OpenHeaderView", "Preview Equivalent C++ Header"),
+						LOCTEXT("OpenHeaderViewTooltip", "Provides a preview of what this class could look like in C++"),
 						FSlateIcon(FEditorStyle::GetStyleSetName(), "ClassIcon.Class"),
 						FUIAction(FExecuteAction::CreateStatic(&FBlueprintHeaderViewModule::OpenHeaderViewForAsset, SelectedAssets[0]))
 						);
