BP Header View: Add Sorting option for optimal struct padding

diff --git a/Engine/Plugins/Editor/BlueprintHeaderView/Source/BlueprintHeaderView/Private/SBlueprintHeaderView.cpp b/Engine/Plugins/Editor/BlueprintHeaderView/Source/BlueprintHeaderView/Private/SBlueprintHeaderView.cpp
--- a/Engine/Plugins/Editor/BlueprintHeaderView/Source/BlueprintHeaderView/Private/SBlueprintHeaderView.cpp
+++ b/Engine/Plugins/Editor/BlueprintHeaderView/Source/BlueprintHeaderView/Private/SBlueprintHeaderView.cpp
@@ -573,6 +573,67 @@ void SBlueprintHeaderView::GatherProperties(const UBlueprint* Blueprint, TArray<
 				return !LeftProp.GetBoolMetaData(FBlueprintMetadata::MD_Private) && RightProp.GetBoolMetaData(FBlueprintMetadata::MD_Private);
 			});
 	}
+	else if (BlueprintHeaderViewSettings->SortMethod == EHeaderViewSortMethod::SortForOptimalPadding)
+	{
+		SortPropertiesForPadding(OutProperties);
+	}
+}
+
+void SBlueprintHeaderView::SortPropertiesForPadding(TArray<const FProperty*>& InOutProperties)
+{
+	TArray<const FProperty*> SortedProperties;
+	SortedProperties.Reserve(InOutProperties.Num());
+
+	auto GetNeededPadding = [] (int32 Offset, int32 Alignment)
+	{
+		if (Offset % Alignment == 0)
+		{
+			return 0;
+		}
+
+		return Alignment - (Offset % Alignment);
+	};
+
+	/**
+	 * For optimal struct size, we always want to place a property that can be aligned
+	 * on the current boundary with the least amount of padding. If there are multiple
+	 * properties that could be placed with the same amount of padding, we should pick
+	 * the one with the greatest alignment because that one will be harder to place
+	 * perfectly.
+	 * 
+	 * Any types' size will be a multiple of its alignment, so an 8-aligned property
+	 * can always be followed by a 4-aligned property with no padding, so if we're 
+	 * at an 8-byte boundary we should always put the 8 first.
+	 */
+
+	int32 CurrentOffset = 0;
+	while (!InOutProperties.IsEmpty())
+	{
+		int32 BestIndex = INDEX_NONE;
+		int32 BestAlignment = 0;
+		int32 BestPadding = 0;
+		for (int32 PropIndex = 0; PropIndex < InOutProperties.Num(); ++PropIndex)
+		{
+			const FProperty* Property = InOutProperties[PropIndex];
+			const int32 PropAlignment = Property->GetMinAlignment();
+			const int32 PropPadding = GetNeededPadding(CurrentOffset, PropAlignment);
+
+			if (BestIndex == INDEX_NONE
+				|| PropPadding < BestPadding 
+				|| (PropPadding == BestPadding && PropAlignment > BestAlignment))
+			{
+				BestIndex = PropIndex;
+				BestAlignment = PropAlignment;
+				BestPadding = PropPadding;
+			}
+		}
+
+		SortedProperties.Add(InOutProperties[BestIndex]);
+		InOutProperties.RemoveAt(BestIndex, 1, /*bAllowShrinking=*/false);
+		CurrentOffset += BestPadding + SortedProperties.Last()->GetSize();
+	}
+
+	InOutProperties = MoveTemp(SortedProperties);
 }
 
 TSharedPtr<SWidget> SBlueprintHeaderView::OnContextMenuOpening()
diff --git a/Engine/Plugins/Editor/BlueprintHeaderView/Source/BlueprintHeaderView/Private/SBlueprintHeaderView.h b/Engine/Plugins/Editor/BlueprintHeaderView/Source/BlueprintHeaderView/Private/SBlueprintHeaderView.h
--- a/Engine/Plugins/Editor/BlueprintHeaderView/Source/BlueprintHeaderView/Private/SBlueprintHeaderView.h
+++ b/Engine/Plugins/Editor/BlueprintHeaderView/Source/BlueprintHeaderView/Private/SBlueprintHeaderView.h
@@ -139,6 +139,9 @@ private:
 	/** Gathers all properties from the blueprint and sorts them according to the selected method from config */
 	void GatherProperties(const UBlueprint* Blueprint, TArray<const FProperty*>& OutProperties);
 
+	/** Sorts an array of properties to optimize for minimal C++ struct padding */
+	void SortPropertiesForPadding(TArray<const FProperty*>& InOutProperties);
+
 	/** Creates a context menu for the list view */
 	TSharedPtr<SWidget> OnContextMenuOpening();
 
diff --git a/Engine/Plugins/Editor/BlueprintHeaderView/Source/BlueprintHeaderView/Public/BlueprintHeaderViewSettings.h b/Engine/Plugins/Editor/BlueprintHeaderView/Source/BlueprintHeaderView/Public/BlueprintHeaderViewSettings.h
--- a/Engine/Plugins/Editor/BlueprintHeaderView/Source/BlueprintHeaderView/Public/BlueprintHeaderViewSettings.h
+++ b/Engine/Plugins/Editor/BlueprintHeaderView/Source/BlueprintHeaderView/Public/BlueprintHeaderViewSettings.h
@@ -8,8 +8,14 @@
 UENUM()
 enum class EHeaderViewSortMethod : uint8
 {
-	None,
-	SortByAccessSpecifier
+	// Properties will stay in the same order they were in the Blueprint class.
+	None, 
+
+	// Properties will be grouped together by Access Specifiers, in order of visibility (public, protected, private).
+	SortByAccessSpecifier,
+	
+	// Properties will be sorted to minimize padding in compiled class layout.
+	SortForOptimalPadding 
 };
 
 USTRUCT(NotBlueprintable)
