BPHeaderView: Remove OnBlueprintChanged delegate when header view is closed

diff --git a/Engine/Plugins/Editor/BlueprintHeaderView/Source/BlueprintHeaderView/Private/SBlueprintHeaderView.cpp b/Engine/Plugins/Editor/BlueprintHeaderView/Source/BlueprintHeaderView/Private/SBlueprintHeaderView.cpp
--- a/Engine/Plugins/Editor/BlueprintHeaderView/Source/BlueprintHeaderView/Private/SBlueprintHeaderView.cpp
+++ b/Engine/Plugins/Editor/BlueprintHeaderView/Source/BlueprintHeaderView/Private/SBlueprintHeaderView.cpp
@@ -316,6 +316,14 @@ void SBlueprintHeaderView::Construct(const FArguments& InArgs)
 	];
 }
 
+SBlueprintHeaderView::~SBlueprintHeaderView()
+{
+	if (UBlueprint* Blueprint = SelectedBlueprint.Get())
+	{
+		Blueprint->OnChanged().RemoveAll(this);
+	}
+}
+
 FReply SBlueprintHeaderView::OnKeyDown(const FGeometry& MyGeometry, const FKeyEvent& InKeyEvent)
 {
 	if (CommandList.IsValid() && CommandList->ProcessCommandBindings(InKeyEvent))
diff --git a/Engine/Plugins/Editor/BlueprintHeaderView/Source/BlueprintHeaderView/Private/SBlueprintHeaderView.h b/Engine/Plugins/Editor/BlueprintHeaderView/Source/BlueprintHeaderView/Private/SBlueprintHeaderView.h
--- a/Engine/Plugins/Editor/BlueprintHeaderView/Source/BlueprintHeaderView/Private/SBlueprintHeaderView.h
+++ b/Engine/Plugins/Editor/BlueprintHeaderView/Source/BlueprintHeaderView/Private/SBlueprintHeaderView.h
@@ -100,6 +100,8 @@ public:
 
 	void Construct(const FArguments& InArgs);
 
+	~SBlueprintHeaderView();
+
 	//~ SWidget interface
 	virtual FReply OnKeyDown(const FGeometry& MyGeometry, const FKeyEvent& InKeyEvent) override;
 	//~ End SWidget interface
