Clamp Blueprint Header Preview settings Font Size to avoid ensure

diff --git a/Engine/Plugins/Editor/BlueprintHeaderView/Source/BlueprintHeaderView/Public/BlueprintHeaderViewSettings.h b/Engine/Plugins/Editor/BlueprintHeaderView/Source/BlueprintHeaderView/Public/BlueprintHeaderViewSettings.h
--- a/Engine/Plugins/Editor/BlueprintHeaderView/Source/BlueprintHeaderView/Public/BlueprintHeaderViewSettings.h
+++ b/Engine/Plugins/Editor/BlueprintHeaderView/Source/BlueprintHeaderView/Public/BlueprintHeaderViewSettings.h
@@ -77,7 +77,7 @@ public:
 	FLinearColor SelectionColor = FLinearColor(0.3f, 0.3f, 1.0f, 1.0f);
 
 	/** Font Size for the Blueprint Header View output */
-	UPROPERTY(config, EditAnywhere, Category="Settings|Style")
+	UPROPERTY(config, EditAnywhere, Category="Settings|Style", meta=(ClampMin=6, ClampMax=72))
 	int32 FontSize = 9;
 	
 	/** Sorting Method for Header View Functions and Properties */
