BP Header Preview: Add buttons to browse to asset and open asset

diff --git a/Engine/Plugins/Editor/BlueprintHeaderView/Source/BlueprintHeaderView/Private/SBlueprintHeaderView.cpp b/Engine/Plugins/Editor/BlueprintHeaderView/Source/BlueprintHeaderView/Private/SBlueprintHeaderView.cpp
--- a/Engine/Plugins/Editor/BlueprintHeaderView/Source/BlueprintHeaderView/Private/SBlueprintHeaderView.cpp
+++ b/Engine/Plugins/Editor/BlueprintHeaderView/Source/BlueprintHeaderView/Private/SBlueprintHeaderView.cpp
@@ -267,6 +267,36 @@ void SBlueprintHeaderView::Construct(const FArguments& InArgs)
 				]
 			]
 			+SHorizontalBox::Slot()
+			.AutoWidth()
+			.HAlign(HAlign_Left)
+			[
+				SNew(SButton)
+				.ButtonStyle(FEditorStyle::Get(), "HoverHintOnly")
+				.ToolTipText(LOCTEXT("BrowseToBlueprintTooltip", "Browse to Selected Blueprint in Content Browser"))
+				.OnClicked(this, &SBlueprintHeaderView::BrowseToAssetClicked)
+				.ContentPadding(4.0f)
+				.Content()
+				[
+					SNew(SImage)
+					.Image(FAppStyle::Get().GetBrush("Icons.BrowseContent"))
+				]
+			]
+			+SHorizontalBox::Slot()
+			.AutoWidth()
+			.HAlign(HAlign_Left)
+			[
+				SNew(SButton)
+				.ButtonStyle(FEditorStyle::Get(), "HoverHintOnly")
+				.ToolTipText(LOCTEXT("EditBlueprintTooltip", "Open Selected Blueprint in Blueprint Editor"))
+				.OnClicked(this, &SBlueprintHeaderView::OpenAssetEditorClicked)
+				.ContentPadding(4.0f)
+				.Content()
+				[
+					SNew(SImage)
+					.Image(FEditorStyle::GetBrush("Icons.Edit"))
+				]
+			]
+			+SHorizontalBox::Slot()
 			.HAlign(HAlign_Right)
 			[
 				SNew(SComboButton)
@@ -697,4 +727,26 @@ void SBlueprintHeaderView::OnSelectAll()
 	ListView->SetItemSelection(ListItems, true);
 }
 
+FReply SBlueprintHeaderView::BrowseToAssetClicked() const
+{
+	if (UBlueprint* Blueprint = SelectedBlueprint.Get())
+	{
+		TArray<FAssetData> AssetsToSync;
+		AssetsToSync.Emplace(Blueprint);
+		GEditor->SyncBrowserToObjects(AssetsToSync);
+	}
+
+	return FReply::Handled();
+}
+
+FReply SBlueprintHeaderView::OpenAssetEditorClicked() const
+{
+	if (UBlueprint* Blueprint = SelectedBlueprint.Get())
+	{
+		GEditor->GetEditorSubsystem<UAssetEditorSubsystem>()->OpenEditorForAsset(Blueprint);
+	}
+
+	return FReply::Handled();
+}
+
 #undef LOCTEXT_NAMESPACE
diff --git a/Engine/Plugins/Editor/BlueprintHeaderView/Source/BlueprintHeaderView/Private/SBlueprintHeaderView.h b/Engine/Plugins/Editor/BlueprintHeaderView/Source/BlueprintHeaderView/Private/SBlueprintHeaderView.h
--- a/Engine/Plugins/Editor/BlueprintHeaderView/Source/BlueprintHeaderView/Private/SBlueprintHeaderView.h
+++ b/Engine/Plugins/Editor/BlueprintHeaderView/Source/BlueprintHeaderView/Private/SBlueprintHeaderView.h
@@ -157,7 +157,8 @@ private:
 	void OnCopy() const;
 	bool CanCopy() const;
 	void OnSelectAll();
-
+	FReply BrowseToAssetClicked() const;
+	FReply OpenAssetEditorClicked() const;
 private:
 	/** List of UI Commands for this scope */
 	TSharedPtr<FUICommandList> CommandList;
