Fix missing includes

diff --git a/Engine/Plugins/Editor/BlueprintHeaderView/Source/BlueprintHeaderView/Private/SBlueprintHeaderView.cpp b/Engine/Plugins/Editor/BlueprintHeaderView/Source/BlueprintHeaderView/Private/SBlueprintHeaderView.cpp
--- a/Engine/Plugins/Editor/BlueprintHeaderView/Source/BlueprintHeaderView/Private/SBlueprintHeaderView.cpp
+++ b/Engine/Plugins/Editor/BlueprintHeaderView/Source/BlueprintHeaderView/Private/SBlueprintHeaderView.cpp
@@ -7,6 +7,7 @@
 //#include "Widgets/Layout/SScrollBox.h"
 #include "Widgets/Input/SComboButton.h"
 #include "Widgets/Text/SRichTextBlock.h"
+#include "Widgets/Input/SButton.h"
 #include "PropertyEditorModule.h"
 #include "ContentBrowserModule.h"
 #include "IContentBrowserSingleton.h"
@@ -24,11 +25,15 @@
 #include "HAL/PlatformApplicationMisc.h"
 #include "BlueprintHeaderViewSettings.h"
 #include "Internationalization/Regex.h"
+#include "Subsystems/AssetEditorSubsystem.h"
+#include "Editor/EditorEngine.h"
 
 #define LOCTEXT_NAMESPACE "SBlueprintHeaderView"
 
 DEFINE_LOG_CATEGORY(LogBlueprintHeaderView)
 
+extern UNREALED_API UEditorEngine* GEditor;
+
 // HeaderViewSyntaxDecorators /////////////////////////////////////////////////
 
 namespace HeaderViewSyntaxDecorators
