Fix Syntax Color not changing in C++ Header Preview Settings

diff --git a/Engine/Plugins/Editor/BlueprintHeaderView/Source/BlueprintHeaderView/Private/BlueprintHeaderViewSettings.cpp b/Engine/Plugins/Editor/BlueprintHeaderView/Source/BlueprintHeaderView/Private/BlueprintHeaderViewSettings.cpp
--- a/Engine/Plugins/Editor/BlueprintHeaderView/Source/BlueprintHeaderView/Private/BlueprintHeaderViewSettings.cpp
+++ b/Engine/Plugins/Editor/BlueprintHeaderView/Source/BlueprintHeaderView/Private/BlueprintHeaderViewSettings.cpp
@@ -41,8 +41,6 @@ void UBlueprintHeaderViewSettings::PostEditChangeProperty(FPropertyChangedEvent&
 
 	FName PropertyChangedName = PropertyChangedEvent.GetPropertyName();
 
-	SyntaxColors = FHeaderViewSyntaxColors();
-
 	if (PropertyChangedName == GET_MEMBER_NAME_CHECKED(UBlueprintHeaderViewSettings, FontSize))
 	{
 		FBlueprintHeaderViewModule::HeaderViewTextStyle.SetFontSize(FontSize);
