Save asset picker config in the Header Preview

diff --git a/Engine/Plugins/Editor/BlueprintHeaderView/Source/BlueprintHeaderView/Private/SBlueprintHeaderView.cpp b/Engine/Plugins/Editor/BlueprintHeaderView/Source/BlueprintHeaderView/Private/SBlueprintHeaderView.cpp
--- a/Engine/Plugins/Editor/BlueprintHeaderView/Source/BlueprintHeaderView/Private/SBlueprintHeaderView.cpp
+++ b/Engine/Plugins/Editor/BlueprintHeaderView/Source/BlueprintHeaderView/Private/SBlueprintHeaderView.cpp
@@ -398,6 +398,7 @@ TSharedRef<SWidget> SBlueprintHeaderView::GetClassPickerMenuContent()
 	AssetPickerConfig.OnAssetSelected = FOnAssetSelected::CreateSP(this, &SBlueprintHeaderView::OnAssetSelected);
 	AssetPickerConfig.Filter.ClassNames.Add(UBlueprint::StaticClass()->GetFName());
 	AssetPickerConfig.InitialAssetViewType = EAssetViewType::List;
+	AssetPickerConfig.SaveSettingsName = TEXT("HeaderPreviewClassPickerSettings");
 	TSharedRef<SWidget> AssetPickerWidget = ContentBrowserModule.Get().CreateAssetPicker(AssetPickerConfig);
 
 	return SNew(SBox)
