Enable Blueprint Header Preview plugin by default

diff --git a/Engine/Plugins/Editor/BlueprintHeaderView/BlueprintHeaderView.uplugin b/Engine/Plugins/Editor/BlueprintHeaderView/BlueprintHeaderView.uplugin
--- a/Engine/Plugins/Editor/BlueprintHeaderView/BlueprintHeaderView.uplugin
+++ b/Engine/Plugins/Editor/BlueprintHeaderView/BlueprintHeaderView.uplugin
@@ -10,6 +10,7 @@
 	"DocsURL": "",
 	"MarketplaceURL": "",
 	"SupportURL": "",
+	"EnabledByDefault": true,
 	"CanContainContent": true,
 	"CanContainVerse": false,
 	"IsBetaVersion": false,
