Check MemberProperty name instead of ChangedProperty name in UBlueprintHeaderViewSettings::PostEditChangeProperty

diff --git a/Engine/Plugins/Editor/BlueprintHeaderView/Source/BlueprintHeaderView/Private/BlueprintHeaderViewSettings.cpp b/Engine/Plugins/Editor/BlueprintHeaderView/Source/BlueprintHeaderView/Private/BlueprintHeaderViewSettings.cpp
--- a/Engine/Plugins/Editor/BlueprintHeaderView/Source/BlueprintHeaderView/Private/BlueprintHeaderViewSettings.cpp
+++ b/Engine/Plugins/Editor/BlueprintHeaderView/Source/BlueprintHeaderView/Private/BlueprintHeaderViewSettings.cpp
@@ -39,7 +39,12 @@ void UBlueprintHeaderViewSettings::PostEditChangeProperty(FPropertyChangedEvent&
 {
 	Super::PostEditChangeProperty(PropertyChangedEvent);
 
-	FName PropertyChangedName = PropertyChangedEvent.GetPropertyName();
+	if (!PropertyChangedEvent.MemberProperty)
+	{
+		return;
+	}
+
+	const FName PropertyChangedName = PropertyChangedEvent.MemberProperty->GetFName();
 
 	if (PropertyChangedName == GET_MEMBER_NAME_CHECKED(UBlueprintHeaderViewSettings, FontSize))
 	{
