Add UserDefinedStruct support to BPHeaderView

diff --git a/Engine/Plugins/Editor/BlueprintHeaderView/Source/BlueprintHeaderView/Private/BlueprintHeaderView.cpp b/Engine/Plugins/Editor/BlueprintHeaderView/Source/BlueprintHeaderView/Private/BlueprintHeaderView.cpp
--- a/Engine/Plugins/Editor/BlueprintHeaderView/Source/BlueprintHeaderView/Private/BlueprintHeaderView.cpp
+++ b/Engine/Plugins/Editor/BlueprintHeaderView/Source/BlueprintHeaderView/Private/BlueprintHeaderView.cpp
@@ -15,6 +15,7 @@
 #include "ContentBrowserModule.h"
 #include "Framework/Commands/UIAction.h"
 #include "Engine/Blueprint.h"
+#include "Engine/UserDefinedStruct.h"
 #include "Framework/MultiBox/MultiBoxBuilder.h"
 #include "ToolMenus.h"
 #include "Toolkits/AssetEditorToolkit.h"
@@ -75,7 +76,7 @@ void FBlueprintHeaderViewModule::ShutdownModule()
 	
 bool FBlueprintHeaderViewModule::IsClassHeaderViewSupported(const UClass* InClass)
 {
-	return InClass == UBlueprint::StaticClass();
+	return InClass == UBlueprint::StaticClass() || InClass == UUserDefinedStruct::StaticClass();
 }
 
 void FBlueprintHeaderViewModule::OpenHeaderViewForAsset(FAssetData InAssetData)
diff --git a/Engine/Plugins/Editor/BlueprintHeaderView/Source/BlueprintHeaderView/Private/HeaderViewClassListItem.cpp b/Engine/Plugins/Editor/BlueprintHeaderView/Source/BlueprintHeaderView/Private/HeaderViewClassListItem.cpp
--- a/Engine/Plugins/Editor/BlueprintHeaderView/Source/BlueprintHeaderView/Private/HeaderViewClassListItem.cpp
+++ b/Engine/Plugins/Editor/BlueprintHeaderView/Source/BlueprintHeaderView/Private/HeaderViewClassListItem.cpp
@@ -15,18 +15,19 @@ FHeaderViewListItemPtr FHeaderViewClassListItem::Create(TWeakObjectPtr<UBlueprin
 	return MakeShareable(new FHeaderViewClassListItem(InBlueprint));
 }
 
-void FHeaderViewClassListItem::ExtendContextMenu(FMenuBuilder& InMenuBuilder, TWeakObjectPtr<UBlueprint> InBlueprint)
+void FHeaderViewClassListItem::ExtendContextMenu(FMenuBuilder& InMenuBuilder, TWeakObjectPtr<UObject> InAsset)
 {
 	if (!bIsValidName)
 	{
-		if (UBlueprint* Blueprint = InBlueprint.Get())
+		if (UBlueprint* Blueprint = Cast<UBlueprint>(InAsset.Get()))
 		{
+			TWeakObjectPtr<UBlueprint> WeakBlueprint = Blueprint;
 			InMenuBuilder.AddVerifiedEditableText(LOCTEXT("RenameBlueprint", "Rename Blueprint"),
 				LOCTEXT("RenameItemTooltip", "Renames this Blueprint\nThis Blueprint name is not a legal C++ identifier."),
 				FSlateIcon(),
 				FText::FromString(Blueprint->GetName()),
-				FOnVerifyTextChanged::CreateSP(this, &FHeaderViewClassListItem::OnVerifyRenameTextChanged, InBlueprint),
-				FOnTextCommitted::CreateSP(this, &FHeaderViewClassListItem::OnRenameTextComitted, InBlueprint)
+				FOnVerifyTextChanged::CreateSP(this, &FHeaderViewClassListItem::OnVerifyRenameTextChanged, WeakBlueprint),
+				FOnTextCommitted::CreateSP(this, &FHeaderViewClassListItem::OnRenameTextComitted, WeakBlueprint)
 			);
 		}
 	}
diff --git a/Engine/Plugins/Editor/BlueprintHeaderView/Source/BlueprintHeaderView/Private/HeaderViewClassListItem.h b/Engine/Plugins/Editor/BlueprintHeaderView/Source/BlueprintHeaderView/Private/HeaderViewClassListItem.h
--- a/Engine/Plugins/Editor/BlueprintHeaderView/Source/BlueprintHeaderView/Private/HeaderViewClassListItem.h
+++ b/Engine/Plugins/Editor/BlueprintHeaderView/Source/BlueprintHeaderView/Private/HeaderViewClassListItem.h
@@ -15,7 +15,7 @@ public:
 	static FHeaderViewListItemPtr Create(TWeakObjectPtr<UBlueprint> InBlueprint);
 
 	//~ FHeaderViewListItem Interface
-	virtual void ExtendContextMenu(FMenuBuilder& InMenuBuilder, TWeakObjectPtr<UBlueprint> InBlueprint) override;
+	virtual void ExtendContextMenu(FMenuBuilder& InMenuBuilder, TWeakObjectPtr<UObject> InAsset) override;
 	//~ End FHeaderViewListItem Interface
 protected:
 	FString GetConditionalUClassSpecifiers(const UBlueprint* Blueprint) const;
diff --git a/Engine/Plugins/Editor/BlueprintHeaderView/Source/BlueprintHeaderView/Private/HeaderViewFunctionListItem.cpp b/Engine/Plugins/Editor/BlueprintHeaderView/Source/BlueprintHeaderView/Private/HeaderViewFunctionListItem.cpp
--- a/Engine/Plugins/Editor/BlueprintHeaderView/Source/BlueprintHeaderView/Private/HeaderViewFunctionListItem.cpp
+++ b/Engine/Plugins/Editor/BlueprintHeaderView/Source/BlueprintHeaderView/Private/HeaderViewFunctionListItem.cpp
@@ -16,40 +16,45 @@ FHeaderViewListItemPtr FHeaderViewFunctionListItem::Create(const UK2Node_Functio
 	return MakeShareable(new FHeaderViewFunctionListItem(FunctionEntry));
 }
 
-void FHeaderViewFunctionListItem::ExtendContextMenu(FMenuBuilder& InMenuBuilder, TWeakObjectPtr<UBlueprint> InBlueprint)
+void FHeaderViewFunctionListItem::ExtendContextMenu(FMenuBuilder& InMenuBuilder, TWeakObjectPtr<UObject> InAsset)
 {
-	InMenuBuilder.AddMenuEntry(LOCTEXT("JumpToDefinition", "Jump to Definition"),
-		LOCTEXT("JumpToDefinitionTooltip", "Opens this function in the Blueprint Editor"),
-		FSlateIcon(),
-		FUIAction(FExecuteAction::CreateSP(this, &FHeaderViewFunctionListItem::JumpToDefinition, InBlueprint))
-	);
-
-	if (!IllegalName.IsNone())
+	if (UBlueprint* Blueprint = Cast<UBlueprint>(InAsset.Get()))
 	{
-		InMenuBuilder.AddVerifiedEditableText(LOCTEXT("RenameItem", "Rename Function"),
-			LOCTEXT("RenameItemTooltip", "Renames this function in the Blueprint\nThis function name is not a legal C++ identifier."),
+		TWeakObjectPtr<UBlueprint> WeakBlueprint = Blueprint;
+		InMenuBuilder.AddMenuEntry(LOCTEXT("JumpToDefinition", "Jump to Definition"),
+			LOCTEXT("JumpToDefinitionTooltip", "Opens this function in the Blueprint Editor"),
 			FSlateIcon(),
-			FText::FromName(IllegalName),
-			FOnVerifyTextChanged::CreateSP(this, &FHeaderViewFunctionListItem::OnVerifyRenameFunctionTextChanged, InBlueprint),
-			FOnTextCommitted::CreateSP(this, &FHeaderViewFunctionListItem::OnRenameFunctionTextCommitted, InBlueprint, GraphName)
+			FUIAction(FExecuteAction::CreateSP(this, &FHeaderViewFunctionListItem::JumpToDefinition, WeakBlueprint))
 		);
-	}
 
-	for (FName IllegalParam : IllegalParameters)
-	{
-		InMenuBuilder.AddVerifiedEditableText(LOCTEXT("RenameParm", "Rename Parameter"),
-			LOCTEXT("RenameParmTooltip", "Renames this function parameter in the Blueprint\nThis parameter name is not a legal C++ identifier."),
-			FSlateIcon(),
-			FText::FromName(IllegalParam),
-			FOnVerifyTextChanged::CreateSP(this, &FHeaderViewFunctionListItem::OnVerifyRenameParameterTextChanged, InBlueprint, GraphName),
-			FOnTextCommitted::CreateSP(this, &FHeaderViewFunctionListItem::OnRenameParameterTextCommitted, InBlueprint, GraphName, IllegalParam)
-		);
+		if (!IllegalName.IsNone())
+		{
+			InMenuBuilder.AddVerifiedEditableText(LOCTEXT("RenameItem", "Rename Function"),
+				LOCTEXT("RenameItemTooltip", "Renames this function in the Blueprint\nThis function name is not a legal C++ identifier."),
+				FSlateIcon(),
+				FText::FromName(IllegalName),
+				FOnVerifyTextChanged::CreateSP(this, &FHeaderViewFunctionListItem::OnVerifyRenameFunctionTextChanged, WeakBlueprint),
+				FOnTextCommitted::CreateSP(this, &FHeaderViewFunctionListItem::OnRenameFunctionTextCommitted, WeakBlueprint, GraphName)
+			);
+		}
+
+		for (FName IllegalParam : IllegalParameters)
+		{
+			InMenuBuilder.AddVerifiedEditableText(LOCTEXT("RenameParm", "Rename Parameter"),
+				LOCTEXT("RenameParmTooltip", "Renames this function parameter in the Blueprint\nThis parameter name is not a legal C++ identifier."),
+				FSlateIcon(),
+				FText::FromName(IllegalParam),
+				FOnVerifyTextChanged::CreateSP(this, &FHeaderViewFunctionListItem::OnVerifyRenameParameterTextChanged, WeakBlueprint, GraphName),
+				FOnTextCommitted::CreateSP(this, &FHeaderViewFunctionListItem::OnRenameParameterTextCommitted, WeakBlueprint, GraphName, IllegalParam)
+			);
+		}
 	}
 }
 
-void FHeaderViewFunctionListItem::OnMouseButtonDoubleClick(TWeakObjectPtr<UBlueprint> Blueprint)
+void FHeaderViewFunctionListItem::OnMouseButtonDoubleClick(TWeakObjectPtr<UObject> InAsset)
 {
-	JumpToDefinition(Blueprint);
+	TWeakObjectPtr<UBlueprint> WeakBlueprint = Cast<UBlueprint>(InAsset.Get());
+	JumpToDefinition(WeakBlueprint);
 }
 
 FString FHeaderViewFunctionListItem::GetConditionalUFunctionSpecifiers(const UFunction* SigFunction) const
diff --git a/Engine/Plugins/Editor/BlueprintHeaderView/Source/BlueprintHeaderView/Private/HeaderViewFunctionListItem.h b/Engine/Plugins/Editor/BlueprintHeaderView/Source/BlueprintHeaderView/Private/HeaderViewFunctionListItem.h
--- a/Engine/Plugins/Editor/BlueprintHeaderView/Source/BlueprintHeaderView/Private/HeaderViewFunctionListItem.h
+++ b/Engine/Plugins/Editor/BlueprintHeaderView/Source/BlueprintHeaderView/Private/HeaderViewFunctionListItem.h
@@ -14,8 +14,8 @@ struct FHeaderViewFunctionListItem : public FHeaderViewListItem
 	static FHeaderViewListItemPtr Create(const UK2Node_FunctionEntry* FunctionEntry);
 
 	//~ FHeaderViewListItem Interface
-	virtual void ExtendContextMenu(FMenuBuilder& InMenuBuilder, TWeakObjectPtr<UBlueprint> InBlueprint) override;
-	virtual void OnMouseButtonDoubleClick(TWeakObjectPtr<UBlueprint> Blueprint) override;
+	virtual void ExtendContextMenu(FMenuBuilder& InMenuBuilder, TWeakObjectPtr<UObject> InAsset) override;
+	virtual void OnMouseButtonDoubleClick(TWeakObjectPtr<UObject> InAsset) override;
 	//~ End FHeaderViewListItem Interface
 
 protected:
diff --git a/Engine/Plugins/Editor/BlueprintHeaderView/Source/BlueprintHeaderView/Private/HeaderViewStructListItem.cpp b/Engine/Plugins/Editor/BlueprintHeaderView/Source/BlueprintHeaderView/Private/HeaderViewStructListItem.cpp
new file mode 100644
--- /dev/null
+++ b/Engine/Plugins/Editor/BlueprintHeaderView/Source/BlueprintHeaderView/Private/HeaderViewStructListItem.cpp
@@ -0,0 +1,141 @@
+// Copyright Epic Games, Inc. All Rights Reserved.
+
+#include "HeaderViewStructListItem.h"
+#include "Engine/UserDefinedStruct.h"
+#include "String/LineEndings.h"
+#include "Kismet2/StructureEditorUtils.h"
+#include "Framework/MultiBox/MultiBoxBuilder.h"
+#include "AssetViewUtils.h"
+#include "AssetToolsModule.h"
+
+#define LOCTEXT_NAMESPACE "FHeaderViewStructListItem"
+
+FHeaderViewListItemPtr FHeaderViewStructListItem::Create(TWeakObjectPtr<UUserDefinedStruct> InStruct)
+{
+	return MakeShareable(new FHeaderViewStructListItem(InStruct));
+}
+
+void FHeaderViewStructListItem::ExtendContextMenu(FMenuBuilder& InMenuBuilder, TWeakObjectPtr<UObject> InAsset)
+{
+	if (!bIsValidName)
+	{
+		if (UUserDefinedStruct* Struct = Cast<UUserDefinedStruct>(InAsset.Get()))
+		{
+			TWeakObjectPtr<UUserDefinedStruct> WeakStruct = Struct;
+			InMenuBuilder.AddVerifiedEditableText(LOCTEXT("RenameStruct", "Rename Struct"),
+				LOCTEXT("RenameItemTooltip", "Renames this Struct\nThis Struct name is not a legal C++ identifier."),
+				FSlateIcon(),
+				FText::FromString(Struct->GetName()),
+				FOnVerifyTextChanged::CreateSP(this, &FHeaderViewStructListItem::OnVerifyRenameTextChanged, WeakStruct),
+				FOnTextCommitted::CreateSP(this, &FHeaderViewStructListItem::OnRenameTextComitted, WeakStruct)
+ 			);
+		}
+	}
+}
+
+FHeaderViewStructListItem::FHeaderViewStructListItem(TWeakObjectPtr<UUserDefinedStruct> InStruct)
+{
+	if (UUserDefinedStruct* Struct = InStruct.Get())
+	{
+		// Avoid lots of reallocations
+		RawItemString.Reserve(512);
+		RichTextString.Reserve(512);
+
+		// Format comment
+		{
+			FString Comment = FStructureEditorUtils::GetTooltip(Struct);
+			if (Comment.IsEmpty())
+			{
+				Comment = TEXT("Please add a struct description");
+			} 
+			FormatCommentString(Comment, RawItemString, RichTextString);
+		}
+
+		// Add the USTRUCT specifiers
+		// i.e. USTRUCT(BlueprintType)
+		{
+			// Always add BlueprintType
+			RawItemString += TEXT("\nUSTRUCT(BlueprintType)");
+			RichTextString += FString::Printf(TEXT("\n<%s>USTRUCT</>(BlueprintType)"), *HeaderViewSyntaxDecorators::MacroDecorator);
+		}
+
+		// Add the class declaration line
+		// i.e. struct StructName
+		{
+			const FString StructName = Struct->GetPrefixCPP() + Struct->GetName();
+
+			bIsValidName = IsValidCPPIdentifier(StructName);
+			const FString& NameDecorator = bIsValidName ? HeaderViewSyntaxDecorators::TypenameDecorator : HeaderViewSyntaxDecorators::ErrorDecorator;
+
+			RawItemString += FString::Printf(TEXT("\nstruct %s\n{\n\tGENERATED_BODY()"), *StructName);
+			RichTextString += FString::Printf(TEXT("\n<%s>struct</> <%s>%s</>\n{\n\t<%s>GENERATED_BODY</>()"), 
+				*HeaderViewSyntaxDecorators::KeywordDecorator,
+				*NameDecorator, 
+				*StructName,
+				*HeaderViewSyntaxDecorators::MacroDecorator
+			);
+		}
+
+		// normalize to platform newlines
+		UE::String::ToHostLineEndingsInline(RawItemString);
+		UE::String::ToHostLineEndingsInline(RichTextString);
+	}
+}
+
+FString FHeaderViewStructListItem::GetRenamedStructPath(const UUserDefinedStruct* Struct, const FString& NewName) const
+{
+	check(Struct);
+	FString NewPath = Struct->GetPathName();
+	int32 Index = 0;
+	if (!ensure(NewPath.FindLastChar(TEXT('/'), Index)))
+	{
+		// AssetViewUtils::IsValidObjectPathForCreate will return false for an empty string
+		return TEXT("");
+	}
+	NewPath.LeftInline(Index);
+	NewPath /= FString::Printf(TEXT("%s.%s"), *NewName, *NewName);
+	return NewPath;
+}
+
+bool FHeaderViewStructListItem::OnVerifyRenameTextChanged(const FText& InNewName, FText& OutErrorText, TWeakObjectPtr<UUserDefinedStruct> InStruct)
+{
+	if (const UUserDefinedStruct* Struct = InStruct.Get())
+	{
+		const FString NewPath = GetRenamedStructPath(Struct, InNewName.ToString());
+		if (!AssetViewUtils::IsValidObjectPathForCreate(NewPath, Struct->GetClass(), OutErrorText))
+		{
+			return false;
+		}
+
+		if (!IsValidCPPIdentifier(InNewName.ToString()))
+		{
+			OutErrorText = InvalidCPPIdentifierErrorText;
+			return false;
+		}
+
+		return true;
+	}
+
+	return false;
+}
+
+void FHeaderViewStructListItem::OnRenameTextComitted(const FText& CommittedText, ETextCommit::Type TextCommitType, TWeakObjectPtr<UUserDefinedStruct> InStruct)
+{
+	if (TextCommitType == ETextCommit::OnEnter)
+	{
+		FText ErrorText;
+		if (OnVerifyRenameTextChanged(CommittedText, ErrorText, InStruct))
+		{
+			if (UUserDefinedStruct* Struct = InStruct.Get())
+			{
+				const FString NewPath = GetRenamedStructPath(Struct, CommittedText.ToString());
+				FAssetToolsModule& AssetToolsModule = FModuleManager::LoadModuleChecked<FAssetToolsModule>("AssetTools");
+				TArray<FAssetRenameData> AssetToRename = { FAssetRenameData(Struct, NewPath) };
+				AssetToolsModule.Get().RenameAssets(AssetToRename);
+				FStructureEditorUtils::OnStructureChanged(Struct, FStructureEditorUtils::Unknown);
+			}
+		}
+	}
+}
+
+#undef LOCTEXT_NAMESPACE
diff --git a/Engine/Plugins/Editor/BlueprintHeaderView/Source/BlueprintHeaderView/Private/HeaderViewStructListItem.h b/Engine/Plugins/Editor/BlueprintHeaderView/Source/BlueprintHeaderView/Private/HeaderViewStructListItem.h
new file mode 100644
--- /dev/null
+++ b/Engine/Plugins/Editor/BlueprintHeaderView/Source/BlueprintHeaderView/Private/HeaderViewStructListItem.h
@@ -0,0 +1,32 @@
+// Copyright Epic Games, Inc. All Rights Reserved.
+
+#pragma once
+
+#include "CoreMinimal.h"
+#include "SBlueprintHeaderView.h"
+
+class FMenuBuilder;
+
+/** A header view list item that displays the struct declaration */
+struct FHeaderViewStructListItem : public FHeaderViewListItem
+{
+public:
+	/** Creates a list item for the Header view representing a class declaration for the given blueprint */
+	static FHeaderViewListItemPtr Create(TWeakObjectPtr<UUserDefinedStruct> InStruct);
+
+	//~ FHeaderViewListItem Interface
+	virtual void ExtendContextMenu(FMenuBuilder& InMenuBuilder, TWeakObjectPtr<UObject> InAsset) override;
+	//~ End FHeaderViewListItem Interface
+protected:
+	FHeaderViewStructListItem(TWeakObjectPtr<UUserDefinedStruct> InStruct);
+
+	FString GetRenamedStructPath(const UUserDefinedStruct* Struct, const FString& NewName) const;
+
+	bool OnVerifyRenameTextChanged(const FText& InNewName, FText& OutErrorText, TWeakObjectPtr<UUserDefinedStruct> InStruct);
+
+	void OnRenameTextComitted(const FText& CommittedText, ETextCommit::Type TextCommitType, TWeakObjectPtr<UUserDefinedStruct> InStruct);
+
+protected:
+	/** Whether this class name is valid C++ (no spaces, special chars, etc) */
+	bool bIsValidName = true;
+};
\ No newline at end of file
diff --git a/Engine/Plugins/Editor/BlueprintHeaderView/Source/BlueprintHeaderView/Private/HeaderViewVariableListItem.cpp b/Engine/Plugins/Editor/BlueprintHeaderView/Source/BlueprintHeaderView/Private/HeaderViewVariableListItem.cpp
--- a/Engine/Plugins/Editor/BlueprintHeaderView/Source/BlueprintHeaderView/Private/HeaderViewVariableListItem.cpp
+++ b/Engine/Plugins/Editor/BlueprintHeaderView/Source/BlueprintHeaderView/Private/HeaderViewVariableListItem.cpp
@@ -4,8 +4,10 @@
 #include "EdMode.h"
 #include "EdGraphSchema_K2.h"
 #include "Engine/Blueprint.h"
+#include "Engine/UserDefinedStruct.h"
 #include "String/LineEndings.h"
 #include "Kismet2/BlueprintEditorUtils.h"
+#include "Kismet2/StructureEditorUtils.h"
 #include "Framework/MultiBox/MultiBoxBuilder.h"
 
 #define LOCTEXT_NAMESPACE "FHeaderViewVariableListItem"
@@ -15,7 +17,7 @@ FHeaderViewListItemPtr FHeaderViewVariableListItem::Create(const FBPVariableDesc
 	return MakeShareable(new FHeaderViewVariableListItem(VariableDesc, VarProperty));
 }
 
-void FHeaderViewVariableListItem::ExtendContextMenu(FMenuBuilder& InMenuBuilder, TWeakObjectPtr<UBlueprint> InBlueprint)
+void FHeaderViewVariableListItem::ExtendContextMenu(FMenuBuilder& InMenuBuilder, TWeakObjectPtr<UObject> InAsset)
 {
 	if (!IllegalName.IsNone())
 	{
@@ -23,8 +25,8 @@ void FHeaderViewVariableListItem::ExtendContextMenu(FMenuBuilder& InMenuBuilder,
 			LOCTEXT("RenameItemTooltip", "Renames this variable in the Blueprint\nThis variable name is not a legal C++ identifier."),
 			FSlateIcon(),
 			FText::FromName(IllegalName),
-			FOnVerifyTextChanged::CreateSP(this, &FHeaderViewVariableListItem::OnVerifyRenameTextChanged, InBlueprint),
-			FOnTextCommitted::CreateSP(this, &FHeaderViewVariableListItem::OnRenameTextCommitted, InBlueprint)
+			FOnVerifyTextChanged::CreateSP(this, &FHeaderViewVariableListItem::OnVerifyRenameTextChanged, InAsset),
+			FOnTextCommitted::CreateSP(this, &FHeaderViewVariableListItem::OnRenameTextCommitted, InAsset)
 		);
 	}
 }
@@ -329,16 +331,18 @@ FString FHeaderViewVariableListItem::GetOwningClassName(const FProperty& VarProp
 	return TEXT("");
 }
 
-bool FHeaderViewVariableListItem::OnVerifyRenameTextChanged(const FText& InNewName, FText& OutErrorText, TWeakObjectPtr<UBlueprint> WeakBlueprint)
+bool FHeaderViewVariableListItem::OnVerifyRenameTextChanged(const FText& InNewName, FText& OutErrorText, TWeakObjectPtr<UObject> WeakAsset)
 {
-	if (const UBlueprint* Blueprint = WeakBlueprint.Get())
+	if (!IsValidCPPIdentifier(InNewName.ToString()))
 	{
-		if (!IsValidCPPIdentifier(InNewName.ToString()))
-		{
-			OutErrorText = InvalidCPPIdentifierErrorText;
-			return false;
-		}
+		OutErrorText = InvalidCPPIdentifierErrorText;
+		return false;
+	}
 
+	UObject* Asset = WeakAsset.Get();
+
+	if (const UBlueprint* Blueprint = Cast<UBlueprint>(Asset))
+	{
 		FKismetNameValidator NameValidator(Blueprint, NAME_None, nullptr);
 		const EValidatorResult Result = NameValidator.IsValid(InNewName.ToString());
 		if (Result != EValidatorResult::Ok)
@@ -350,19 +354,31 @@ bool FHeaderViewVariableListItem::OnVerifyRenameTextChanged(const FText& InNewNa
 		return true;
 	}
 
+	if (const UUserDefinedStruct* Struct = Cast<UUserDefinedStruct>(Asset))
+	{
+		return FStructureEditorUtils::IsUniqueVariableFriendlyName(Struct, InNewName.ToString());
+	}
+
 	return false;
 }
 
-void FHeaderViewVariableListItem::OnRenameTextCommitted(const FText& CommittedText, ETextCommit::Type TextCommitType, TWeakObjectPtr<UBlueprint> WeakBlueprint)
+void FHeaderViewVariableListItem::OnRenameTextCommitted(const FText& CommittedText, ETextCommit::Type TextCommitType, TWeakObjectPtr<UObject> WeakAsset)
 {
 	if (TextCommitType == ETextCommit::OnEnter)
 	{
-		if (UBlueprint* Blueprint = WeakBlueprint.Get())
+		if (UObject* Asset = WeakAsset.Get())
 		{
 			const FString& CommittedString = CommittedText.ToString();
 			if (IsValidCPPIdentifier(CommittedString))
 			{
-				FBlueprintEditorUtils::RenameMemberVariable(Blueprint, IllegalName, FName(CommittedString));
+				if (UBlueprint* Blueprint = Cast<UBlueprint>(Asset))
+				{
+					FBlueprintEditorUtils::RenameMemberVariable(Blueprint, IllegalName, FName(CommittedString));
+				}
+				else if (UUserDefinedStruct* Struct = Cast<UUserDefinedStruct>(Asset))
+				{
+					FStructureEditorUtils::RenameVariable(Struct, IllegalName.ToString(), CommittedString);
+				}
 			}
 		}
 	}
diff --git a/Engine/Plugins/Editor/BlueprintHeaderView/Source/BlueprintHeaderView/Private/HeaderViewVariableListItem.h b/Engine/Plugins/Editor/BlueprintHeaderView/Source/BlueprintHeaderView/Private/HeaderViewVariableListItem.h
--- a/Engine/Plugins/Editor/BlueprintHeaderView/Source/BlueprintHeaderView/Private/HeaderViewVariableListItem.h
+++ b/Engine/Plugins/Editor/BlueprintHeaderView/Source/BlueprintHeaderView/Private/HeaderViewVariableListItem.h
@@ -14,7 +14,7 @@ struct FHeaderViewVariableListItem : public FHeaderViewListItem
 	static FHeaderViewListItemPtr Create(const FBPVariableDescription* VariableDesc, const FProperty& VarProperty);
 
 	//~ FHeaderViewListItem Interface
-	virtual void ExtendContextMenu(FMenuBuilder& InMenuBuilder, TWeakObjectPtr<UBlueprint> Blueprint) override;
+	virtual void ExtendContextMenu(FMenuBuilder& InMenuBuilder, TWeakObjectPtr<UObject> InAsset) override;
 	//~ End FHeaderViewListItem Interface
 
 protected:
@@ -29,8 +29,8 @@ protected:
 	/** Returns the name of the owning class */
 	FString GetOwningClassName(const FProperty& VarProperty) const;
 
-	bool OnVerifyRenameTextChanged(const FText& InNewName, FText& OutErrorText, TWeakObjectPtr<UBlueprint> WeakBlueprint);
-	void OnRenameTextCommitted(const FText& CommittedText, ETextCommit::Type TextCommitType, TWeakObjectPtr<UBlueprint> WeakBlueprint);
+	bool OnVerifyRenameTextChanged(const FText& InNewName, FText& OutErrorText, TWeakObjectPtr<UObject> WeakAsset);
+	void OnRenameTextCommitted(const FText& CommittedText, ETextCommit::Type TextCommitType, TWeakObjectPtr<UObject> WeakAsset);
 protected:
 
 	/** None if the name is legal, else holds the name of the variable */
diff --git a/Engine/Plugins/Editor/BlueprintHeaderView/Source/BlueprintHeaderView/Private/SBlueprintHeaderView.cpp b/Engine/Plugins/Editor/BlueprintHeaderView/Source/BlueprintHeaderView/Private/SBlueprintHeaderView.cpp
--- a/Engine/Plugins/Editor/BlueprintHeaderView/Source/BlueprintHeaderView/Private/SBlueprintHeaderView.cpp
+++ b/Engine/Plugins/Editor/BlueprintHeaderView/Source/BlueprintHeaderView/Private/SBlueprintHeaderView.cpp
@@ -12,10 +12,12 @@
 #include "ContentBrowserModule.h"
 #include "IContentBrowserSingleton.h"
 #include "Engine/Blueprint.h"
+#include "Engine/UserDefinedStruct.h"
 #include "Framework/Text/SlateTextRun.h"
 #include "Framework/Commands/GenericCommands.h"
 #include "Framework/MultiBox/MultiBoxBuilder.h"
 #include "HeaderViewClassListItem.h"
+#include "HeaderViewStructListItem.h"
 #include "HeaderViewFunctionListItem.h"
 #include "HeaderViewVariableListItem.h"
 #include "EdGraphSchema_K2.h"
@@ -357,6 +359,10 @@ SBlueprintHeaderView::~SBlueprintHeaderView()
 	{
 		Blueprint->OnChanged().RemoveAll(this);
 	}
+	else if (UUserDefinedStruct* Struct = SelectedStruct.Get())
+	{
+		Struct->OnStructChanged().RemoveAll(this);
+	}
 }
 
 FReply SBlueprintHeaderView::OnKeyDown(const FGeometry& MyGeometry, const FKeyEvent& InKeyEvent)
@@ -385,6 +391,10 @@ FText SBlueprintHeaderView::GetClassPickerText() const
 	{
 		return FText::FromName(Blueprint->GetFName());
 	}
+	else if (const UUserDefinedStruct* Struct = SelectedStruct.Get())
+	{
+		return FText::FromName(Struct->GetFName());
+	}
 
 	return LOCTEXT("ClassPickerPickClass", "Select Blueprint Class");
 }
@@ -397,6 +407,7 @@ TSharedRef<SWidget> SBlueprintHeaderView::GetClassPickerMenuContent()
 	AssetPickerConfig.SelectionMode = ESelectionMode::Single;
 	AssetPickerConfig.OnAssetSelected = FOnAssetSelected::CreateSP(this, &SBlueprintHeaderView::OnAssetSelected);
 	AssetPickerConfig.Filter.ClassPaths.Add(UBlueprint::StaticClass()->GetClassPathName());
+	AssetPickerConfig.Filter.ClassPaths.Add(UUserDefinedStruct::StaticClass()->GetClassPathName());
 	AssetPickerConfig.InitialAssetViewType = EAssetViewType::List;
 	AssetPickerConfig.SaveSettingsName = TEXT("HeaderPreviewClassPickerSettings");
 	TSharedRef<SWidget> AssetPickerWidget = ContentBrowserModule.Get().CreateAssetPicker(AssetPickerConfig);
@@ -430,12 +441,24 @@ void SBlueprintHeaderView::OnAssetSelected(const FAssetData& SelectedAsset)
 	{
 		OldBlueprint->OnChanged().RemoveAll(this);
 	}
+	else if (UUserDefinedStruct* OldStruct = SelectedStruct.Get())
+	{
+		OldStruct->OnStructChanged().RemoveAll(this);
+	}
+
+	SelectedBlueprint = nullptr;
+	SelectedStruct = nullptr;
 
 	if (UBlueprint* Blueprint = Cast<UBlueprint>(SelectedAsset.GetAsset()))
 	{
 		Blueprint->OnChanged().AddRaw(this, &SBlueprintHeaderView::OnBlueprintChanged);
 		SelectedBlueprint = Blueprint;
 	}
+	else if (UUserDefinedStruct* UDS = Cast<UUserDefinedStruct>(SelectedAsset.GetAsset()))
+	{
+		UDS->OnStructChanged().AddRaw(this, &SBlueprintHeaderView::OnStructChanged);
+		SelectedStruct = UDS;
+	}
 
 	RepopulateListView();
 }
@@ -461,7 +484,17 @@ void SBlueprintHeaderView::RepopulateListView()
 
 		PopulateFunctionItems(Blueprint);
 
-		PopulateVariableItems(Blueprint);
+		PopulateVariableItems(Blueprint->SkeletonGeneratedClass);
+
+		// Add the closing brace of the class
+		ListItems.Add(FHeaderViewListItem::Create(TEXT("};"), TEXT("};")));
+	}
+	else if (UUserDefinedStruct* Struct = SelectedStruct.Get())
+	{
+		// Add the struct declaration
+		ListItems.Add(FHeaderViewStructListItem::Create(Struct));
+
+		PopulateVariableItems(Struct);
 
 		// Add the closing brace of the class
 		ListItems.Add(FHeaderViewListItem::Create(TEXT("};"), TEXT("};")));
@@ -551,76 +584,76 @@ void SBlueprintHeaderView::GatherFunctionGraphs(const UBlueprint* Blueprint, TAr
 	}
 }
 
-void SBlueprintHeaderView::PopulateVariableItems(const UBlueprint* Blueprint)
+void SBlueprintHeaderView::AddVariableItems(TArray<const FProperty*> VarProperties)
 {
-	if (Blueprint)
-	{
-		const int32 Private = 2;
-		const int32 Public = 1;
+	const int32 Private = 2;
+	const int32 Public = 1;
 
-		TArray<const FProperty*> VarProperties;
-		GatherProperties(Blueprint, VarProperties);
-
-		// We should only add an access specifier line if the previous variable was a different one
-		int32 PrevAccessSpecifier = 0;
-		for (const FProperty* VarProperty : VarProperties)
-		{
-			const int32 AccessSpecifier = VarProperty->GetBoolMetaData(FBlueprintMetadata::MD_Private) ? Private : Public;
-			if (AccessSpecifier != PrevAccessSpecifier)
+	const UBlueprintHeaderViewSettings* BlueprintHeaderViewSettings = GetDefault<UBlueprintHeaderViewSettings>();
+	if (BlueprintHeaderViewSettings->SortMethod == EHeaderViewSortMethod::SortByAccessSpecifier)
+	{
+		VarProperties.Sort([](const FProperty& LeftProp, const FProperty& RightProp)
 			{
-				switch (AccessSpecifier)
-				{
-				case Public:
-					ListItems.Add(FHeaderViewListItem::Create(TEXT("public:"), FString::Printf(TEXT("<%s>public</>:"), *HeaderViewSyntaxDecorators::KeywordDecorator)));
-					break;
-				case Private:
-					ListItems.Add(FHeaderViewListItem::Create(TEXT("private:"), FString::Printf(TEXT("<%s>private</>:"), *HeaderViewSyntaxDecorators::KeywordDecorator)));
-					break;
-				}
+				return !LeftProp.GetBoolMetaData(FBlueprintMetadata::MD_Private) && RightProp.GetBoolMetaData(FBlueprintMetadata::MD_Private);
+			});
+	}
+	else if (BlueprintHeaderViewSettings->SortMethod == EHeaderViewSortMethod::SortForOptimalPadding)
+	{
+		SortPropertiesForPadding(VarProperties);
+	}
 
-				PrevAccessSpecifier = AccessSpecifier;
-			}
-			else
+	// We should only add an access specifier line if the previous variable was a different one
+	int32 PrevAccessSpecifier = 0;
+	for (const FProperty* VarProperty : VarProperties)
+	{
+		const int32 AccessSpecifier = VarProperty->GetBoolMetaData(FBlueprintMetadata::MD_Private) ? Private : Public;
+		if (AccessSpecifier != PrevAccessSpecifier)
+		{
+			switch (AccessSpecifier)
 			{
-				// add an empty line to space variables out
-				ListItems.Add(FHeaderViewListItem::Create(TEXT(""), TEXT("")));
+			case Public:
+				ListItems.Add(FHeaderViewListItem::Create(TEXT("public:"), FString::Printf(TEXT("<%s>public</>:"), *HeaderViewSyntaxDecorators::KeywordDecorator)));
+				break;
+			case Private:
+				ListItems.Add(FHeaderViewListItem::Create(TEXT("private:"), FString::Printf(TEXT("<%s>private</>:"), *HeaderViewSyntaxDecorators::KeywordDecorator)));
+				break;
 			}
 
-			const FBPVariableDescription* VariableDesc = Blueprint->NewVariables.FindByPredicate([&VarProperty](const FBPVariableDescription& Desc)
+			PrevAccessSpecifier = AccessSpecifier;
+		}
+		else
+		{
+			// add an empty line to space variables out
+			ListItems.Add(FHeaderViewListItem::Create(TEXT(""), TEXT("")));
+		}
+
+		FBPVariableDescription* VariableDesc = nullptr;
+		if (UBlueprint* Blueprint = SelectedBlueprint.Get())
+		{
+			VariableDesc = Blueprint->NewVariables.FindByPredicate([&VarProperty](const FBPVariableDescription& Desc)
 				{
 					return Desc.VarName == VarProperty->GetFName();
 				});
-
-			ListItems.Add(FHeaderViewVariableListItem::Create(VariableDesc, *VarProperty));
 		}
+
+		ListItems.Add(FHeaderViewVariableListItem::Create(VariableDesc, *VarProperty));
 	}
 }
 
-void SBlueprintHeaderView::GatherProperties(const UBlueprint* Blueprint, TArray<const FProperty*>& OutProperties)
+void SBlueprintHeaderView::PopulateVariableItems(const UStruct* Struct)
 {
-	for (TFieldIterator<FProperty> PropertyIt(Blueprint->SkeletonGeneratedClass, EFieldIteratorFlags::ExcludeSuper); PropertyIt; ++PropertyIt)
+	TArray<const FProperty*> VarProperties;
+	for (TFieldIterator<FProperty> PropertyIt(Struct, EFieldIteratorFlags::ExcludeSuper); PropertyIt; ++PropertyIt)
 	{
 		if (const FProperty* VarProperty = *PropertyIt)
 		{
 			if (VarProperty->HasAnyPropertyFlags(CPF_BlueprintVisible))
 			{
-				OutProperties.Add(VarProperty);
+				VarProperties.Add(VarProperty);
 			}
 		}
 	}
-
-	const UBlueprintHeaderViewSettings* BlueprintHeaderViewSettings = GetDefault<UBlueprintHeaderViewSettings>();
-	if (BlueprintHeaderViewSettings->SortMethod == EHeaderViewSortMethod::SortByAccessSpecifier)
-	{
-		OutProperties.Sort([](const FProperty& LeftProp, const FProperty& RightProp)
-			{
-				return !LeftProp.GetBoolMetaData(FBlueprintMetadata::MD_Private) && RightProp.GetBoolMetaData(FBlueprintMetadata::MD_Private);
-			});
-	}
-	else if (BlueprintHeaderViewSettings->SortMethod == EHeaderViewSortMethod::SortForOptimalPadding)
-	{
-		SortPropertiesForPadding(OutProperties);
-	}
+	AddVariableItems(VarProperties);
 }
 
 void SBlueprintHeaderView::SortPropertiesForPadding(TArray<const FProperty*>& InOutProperties)
@@ -691,7 +724,8 @@ TSharedPtr<SWidget> SBlueprintHeaderView::OnContextMenuOpening()
 	ListView->GetSelectedItems(SelectedListItems);
 	if (SelectedListItems.Num() == 1)
 	{
-		SelectedListItems[0]->ExtendContextMenu(MenuBuilder, SelectedBlueprint);
+		TWeakObjectPtr<UObject> SelectedAsset = SelectedBlueprint.IsValid() ? (UObject*)SelectedBlueprint.Get() : (UObject*)SelectedStruct.Get();
+		SelectedListItems[0]->ExtendContextMenu(MenuBuilder, SelectedAsset);
 	}
 
 	return MenuBuilder.MakeWidget();
@@ -699,7 +733,8 @@ TSharedPtr<SWidget> SBlueprintHeaderView::OnContextMenuOpening()
 
 void SBlueprintHeaderView::OnItemDoubleClicked(FHeaderViewListItemPtr Item)
 {
-	Item->OnMouseButtonDoubleClick(SelectedBlueprint);
+	TWeakObjectPtr<UObject> SelectedAsset = SelectedBlueprint.IsValid() ? (UObject*)SelectedBlueprint.Get() : (UObject*)SelectedStruct.Get();
+	Item->OnMouseButtonDoubleClick(SelectedAsset);
 }
 
 void SBlueprintHeaderView::OnBlueprintChanged(UBlueprint* InBlueprint)
@@ -707,6 +742,11 @@ void SBlueprintHeaderView::OnBlueprintChanged(UBlueprint* InBlueprint)
 	RepopulateListView();
 }
 
+void SBlueprintHeaderView::OnStructChanged(UUserDefinedStruct* InStruct)
+{
+	RepopulateListView();
+}
+
 void SBlueprintHeaderView::OnCopy() const
 {
 	const int32 StringReserveSize = 2048;
@@ -735,10 +775,18 @@ void SBlueprintHeaderView::OnSelectAll()
 
 FReply SBlueprintHeaderView::BrowseToAssetClicked() const
 {
+	TArray<FAssetData> AssetsToSync;
 	if (UBlueprint* Blueprint = SelectedBlueprint.Get())
 	{
-		TArray<FAssetData> AssetsToSync;
 		AssetsToSync.Emplace(Blueprint);
+	}
+	else if (UUserDefinedStruct* Struct = SelectedStruct.Get())
+	{
+		AssetsToSync.Emplace(Struct);
+	}
+
+	if (!AssetsToSync.IsEmpty())
+	{
 		GEditor->SyncBrowserToObjects(AssetsToSync);
 	}
 
@@ -751,6 +799,10 @@ FReply SBlueprintHeaderView::OpenAssetEditorClicked() const
 	{
 		GEditor->GetEditorSubsystem<UAssetEditorSubsystem>()->OpenEditorForAsset(Blueprint);
 	}
+	else if (UUserDefinedStruct* Struct = SelectedStruct.Get())
+	{
+		GEditor->GetEditorSubsystem<UAssetEditorSubsystem>()->OpenEditorForAsset(Struct);
+	}
 
 	return FReply::Handled();
 }
diff --git a/Engine/Plugins/Editor/BlueprintHeaderView/Source/BlueprintHeaderView/Private/SBlueprintHeaderView.h b/Engine/Plugins/Editor/BlueprintHeaderView/Source/BlueprintHeaderView/Private/SBlueprintHeaderView.h
--- a/Engine/Plugins/Editor/BlueprintHeaderView/Source/BlueprintHeaderView/Private/SBlueprintHeaderView.h
+++ b/Engine/Plugins/Editor/BlueprintHeaderView/Source/BlueprintHeaderView/Private/SBlueprintHeaderView.h
@@ -12,6 +12,8 @@
 
 class SComboButton;
 class UBlueprint;
+class UUserDefinedStruct;
+class UStruct;
 class ITableRow;
 class STableViewBase;
 class FUICommandList;
@@ -45,10 +47,10 @@ struct FHeaderViewListItem : public TSharedFromThis<FHeaderViewListItem>
 	const FString& GetRawItemString() const { return RawItemString; }
 
 	/** Allows the item to add items to the context menu if it is the only item selected */
-	virtual void ExtendContextMenu(FMenuBuilder& InMenuBuilder, TWeakObjectPtr<UBlueprint> InBlueprint) {}
+	virtual void ExtendContextMenu(FMenuBuilder& InMenuBuilder, TWeakObjectPtr<UObject> InAsset) {}
 
 	/** Called when this List Item is double clicked */
-	virtual void OnMouseButtonDoubleClick(TWeakObjectPtr<UBlueprint> Blueprint) {};
+	virtual void OnMouseButtonDoubleClick(TWeakObjectPtr<UObject> InAsset) {};
 
 protected:
 	/** Empty base constructor hidden from public */
@@ -135,12 +137,10 @@ private:
 	/** Gathers all function graphs from the blueprint and sorts them according to the selected method from config */
 	void GatherFunctionGraphs(const UBlueprint* Blueprint, TArray<const UEdGraph*>& OutFunctionGraphs);
 
-	/** Adds items to the list view representing all variables present in the given blueprints */
-	void PopulateVariableItems(const UBlueprint* Blueprint);
+	/** Adds items to the list view representing all variables present in the given asset */
+	void PopulateVariableItems(const UStruct* Struct);
+	void AddVariableItems(TArray<const FProperty*> VarProperties);
 	
-	/** Gathers all properties from the blueprint and sorts them according to the selected method from config */
-	void GatherProperties(const UBlueprint* Blueprint, TArray<const FProperty*>& OutProperties);
-
 	/** Sorts an array of properties to optimize for minimal C++ struct padding */
 	void SortPropertiesForPadding(TArray<const FProperty*>& InOutProperties);
 
@@ -150,8 +150,9 @@ private:
 	/** Called when a List Item is double clicked */
 	void OnItemDoubleClicked(FHeaderViewListItemPtr Item);
 
-	/** Callback for when the selected blueprint is modified */
+	/** Callback for when the selected asset is modified */
 	void OnBlueprintChanged(UBlueprint* InBlueprint);
+	void OnStructChanged(UUserDefinedStruct* InStruct);
 
 	/** UI Command Functions */
 	void OnCopy() const;
@@ -163,8 +164,9 @@ private:
 	/** List of UI Commands for this scope */
 	TSharedPtr<FUICommandList> CommandList;
 
-	/** The blueprint currently being displayed by the header view */
+	/** The asset currently being displayed by the header view */
 	TWeakObjectPtr<UBlueprint> SelectedBlueprint;
+	TWeakObjectPtr<UUserDefinedStruct> SelectedStruct;
 
 	/** Reference to the Class Picker combo button widget */
 	TSharedPtr<SComboButton> ClassPickerComboButton;
diff --git a/Engine/Source/Editor/UnrealEd/Private/Kismet2/StructureEditorUtils.cpp b/Engine/Source/Editor/UnrealEd/Private/Kismet2/StructureEditorUtils.cpp
--- a/Engine/Source/Editor/UnrealEd/Private/Kismet2/StructureEditorUtils.cpp
+++ b/Engine/Source/Editor/UnrealEd/Private/Kismet2/StructureEditorUtils.cpp
@@ -366,6 +366,29 @@ bool FStructureEditorUtils::RenameVariable(UUserDefinedStruct* Struct, FGuid Var
 }
 
 
+bool FStructureEditorUtils::RenameVariable(UUserDefinedStruct* Struct, const FString& OldDisplayNameStr, const FString& NewDisplayNameStr)
+{
+	if (Struct)
+	{
+		if (TArray<FStructVariableDescription>* VarDescArray = GetVarDescPtr(Struct))
+		{
+			FStructVariableDescription* VarDesc = VarDescArray->FindByPredicate(
+				[&OldDisplayNameStr](const FStructVariableDescription& Var)
+				{
+					return Var.FriendlyName == OldDisplayNameStr;
+				}
+			);
+
+			if (VarDesc)
+			{
+				return RenameVariable(Struct, VarDesc->VarGuid, NewDisplayNameStr);
+			}
+		}
+	}
+
+	return false;
+}
+
 bool FStructureEditorUtils::ChangeVariableType(UUserDefinedStruct* Struct, FGuid VarGuid, const FEdGraphPinType& NewType)
 {
 	if (Struct)
@@ -545,6 +568,7 @@ void FStructureEditorUtils::OnStructureChanged(UUserDefinedStruct* Struct, EStru
 		Struct->Status = EUserDefinedStructureStatus::UDSS_Dirty;
 		CompileStructure(Struct);
 		Struct->MarkPackageDirty();
+		Struct->OnChanged();
 	}
 }
 
diff --git a/Engine/Source/Editor/UnrealEd/Public/Kismet2/StructureEditorUtils.h b/Engine/Source/Editor/UnrealEd/Public/Kismet2/StructureEditorUtils.h
--- a/Engine/Source/Editor/UnrealEd/Public/Kismet2/StructureEditorUtils.h
+++ b/Engine/Source/Editor/UnrealEd/Public/Kismet2/StructureEditorUtils.h
@@ -83,6 +83,8 @@ public:
 
 	static bool RenameVariable(UUserDefinedStruct* Struct, FGuid VarGuid, const FString& NewDisplayNameStr);
 
+	static bool RenameVariable(UUserDefinedStruct* Struct, const FString& OldDisplayNameStr, const FString& NewDisplayNameStr);
+
 	static bool ChangeVariableType(UUserDefinedStruct* Struct, FGuid VarGuid, const FEdGraphPinType& NewType);
 
 	static bool ChangeVariableDefaultValue(UUserDefinedStruct* Struct, FGuid VarGuid, const FString& NewDefaultValue);
diff --git a/Engine/Source/Runtime/Engine/Classes/Engine/UserDefinedStruct.h b/Engine/Source/Runtime/Engine/Classes/Engine/UserDefinedStruct.h
--- a/Engine/Source/Runtime/Engine/Classes/Engine/UserDefinedStruct.h
+++ b/Engine/Source/Runtime/Engine/Classes/Engine/UserDefinedStruct.h
@@ -15,6 +15,8 @@ class UUserDefinedStruct;
 class UUserDefinedStructEditorData;
 class UStructCookedMetaData;
 
+DECLARE_MULTICAST_DELEGATE_OneParam(FOnStructChanged, UUserDefinedStruct*);
+
 UENUM()
 enum EUserDefinedStructureStatus
 {
@@ -95,6 +97,8 @@ public:
 	/** Creates a new guid if needed */
 	void ValidateGuid();
 
+	virtual void OnChanged();
+
 	friend UUserDefinedStructEditorData;
 #endif	// WITH_EDITOR
 
@@ -130,6 +134,10 @@ public:
 	void UpdateStructFlags();
 
 #if WITH_EDITORONLY_DATA
+	FORCEINLINE FOnStructChanged& OnStructChanged() { return ChangedEvent; }
+public:
+	FOnStructChanged ChangedEvent;
+
 protected:
 	virtual TSubclassOf<UStructCookedMetaData> GetCookedMetaDataClass() const;
 
diff --git a/Engine/Source/Runtime/Engine/Private/UserDefinedStruct.cpp b/Engine/Source/Runtime/Engine/Private/UserDefinedStruct.cpp
--- a/Engine/Source/Runtime/Engine/Private/UserDefinedStruct.cpp
+++ b/Engine/Source/Runtime/Engine/Private/UserDefinedStruct.cpp
@@ -225,6 +225,11 @@ void UUserDefinedStruct::ValidateGuid()
 	}
 }
 
+void UUserDefinedStruct::OnChanged()
+{
+	ChangedEvent.Broadcast(this);
+}
+
 #endif	// WITH_EDITOR
 
 FProperty* UUserDefinedStruct::CustomFindProperty(const FName Name) const
