Add Custom Icon for Blueprint Header View Plugin

diff --git a/Engine/Plugins/Editor/BlueprintHeaderView/Source/BlueprintHeaderView/Private/BlueprintHeaderView.cpp b/Engine/Plugins/Editor/BlueprintHeaderView/Source/BlueprintHeaderView/Private/BlueprintHeaderView.cpp
--- a/Engine/Plugins/Editor/BlueprintHeaderView/Source/BlueprintHeaderView/Private/BlueprintHeaderView.cpp
+++ b/Engine/Plugins/Editor/BlueprintHeaderView/Source/BlueprintHeaderView/Private/BlueprintHeaderView.cpp
@@ -21,6 +21,10 @@
 #include "Toolkits/AssetEditorToolkit.h"
 #include "Toolkits/AssetEditorToolkitMenuContext.h"
 #include "BlueprintHeaderViewSettings.h"
+#include "Styling/SlateStyle.h"
+#include "Styling/SlateStyleRegistry.h"
+#include "Styling/SlateBrush.h"
+#include "Styling/SlateStyleMacros.h"
 
 #define LOCTEXT_NAMESPACE "BlueprintHeaderViewApp"
 
@@ -40,14 +44,22 @@ namespace BlueprintHeaderViewModule
 
 FTextBlockStyle FBlueprintHeaderViewModule::HeaderViewTextStyle;
 FTableRowStyle FBlueprintHeaderViewModule::HeaderViewTableRowStyle;
+TSharedPtr<FSlateStyleSet> FBlueprintHeaderViewModule::HeaderViewStyleSet;
 
 void FBlueprintHeaderViewModule::StartupModule()
 {
+	const FVector2D Icon16x16 = FVector2D(16.0);
+	const FString PluginContentDir = FPaths::ConvertRelativePathToFull(FPaths::Combine(FPaths::EnginePluginsDir(), TEXT("Editor/BlueprintHeaderView/Content")));
+	HeaderViewStyleSet = MakeShared<FSlateStyleSet>("HeaderViewStyle");
+	HeaderViewStyleSet->SetContentRoot(PluginContentDir);
+	HeaderViewStyleSet->Set("Icons.HeaderView", new FSlateVectorImageBrush(HeaderViewStyleSet->RootToContentDir("BlueprintHeader_16", TEXT(".svg")), Icon16x16));
+	FSlateStyleRegistry::RegisterSlateStyle(*HeaderViewStyleSet);
+
 	FGlobalTabmanager::Get()->RegisterNomadTabSpawner(BlueprintHeaderViewModule::HeaderViewTabName, FOnSpawnTab::CreateStatic(&BlueprintHeaderViewModule::CreateHeaderViewTab))
 		.SetDisplayName(LOCTEXT("TabTitle", "C++ Header Preview"))
 		.SetTooltipText(LOCTEXT("TooltipText", "Displays a Blueprint Class in C++ Header format."))
 		.SetGroup(WorkspaceMenu::GetMenuStructure().GetToolsCategory())
-		.SetIcon(FSlateIcon(FAppStyle::GetAppStyleSetName(), "ClassIcon.Class"));
+		.SetIcon(FSlateIcon(HeaderViewStyleSet->GetStyleSetName(), "Icons.HeaderView"));
 
 	HeaderViewTextStyle = FTextBlockStyle()
 		.SetFont(FCoreStyle::GetDefaultFontStyle("Mono", GetDefault<UBlueprintHeaderViewSettings>()->FontSize))
@@ -72,6 +84,9 @@ void FBlueprintHeaderViewModule::ShutdownModule()
 				return Delegate.GetHandle() == ContentBrowserExtenderDelegateHandle;
 			});
 	}
+
+	FSlateStyleRegistry::UnRegisterSlateStyle(*HeaderViewStyleSet);
+	HeaderViewStyleSet.Reset();
 }
 	
 bool FBlueprintHeaderViewModule::IsClassHeaderViewSupported(const UClass* InClass)
@@ -112,7 +127,7 @@ void FBlueprintHeaderViewModule::SetupAssetEditorMenuExtender()
 								FName("OpenHeaderView"),
 								LOCTEXT("OpenAssetHeaderView", "Preview Equivalent C++ Header"),
 								LOCTEXT("OpenAssetHeaderViewTooltip", "Provides a preview of what this class could look like in C++"),
-								FSlateIcon(FAppStyle::GetAppStyleSetName(), "ClassIcon.Class"),
+								FSlateIcon(HeaderViewStyleSet->GetStyleSetName(), "Icons.HeaderView"),
 								FUIAction(FExecuteAction::CreateStatic(&FBlueprintHeaderViewModule::OpenHeaderViewForAsset, BlueprintAssetData))
 							);
 						}
@@ -143,7 +158,7 @@ TSharedRef<FExtender> FBlueprintHeaderViewModule::OnExtendContentBrowserAssetSel
 					MenuBuilder.AddMenuEntry(
 						LOCTEXT("OpenHeaderView", "Preview Equivalent C++ Header"),
 						LOCTEXT("OpenHeaderViewTooltip", "Provides a preview of what this class could look like in C++"),
-						FSlateIcon(FAppStyle::GetAppStyleSetName(), "ClassIcon.Class"),
+						FSlateIcon(HeaderViewStyleSet->GetStyleSetName(), "Icons.HeaderView"),
 						FUIAction(FExecuteAction::CreateStatic(&FBlueprintHeaderViewModule::OpenHeaderViewForAsset, SelectedAssets[0]))
 						);
 				})
diff --git a/Engine/Plugins/Editor/BlueprintHeaderView/Source/BlueprintHeaderView/Public/BlueprintHeaderView.h b/Engine/Plugins/Editor/BlueprintHeaderView/Source/BlueprintHeaderView/Public/BlueprintHeaderView.h
--- a/Engine/Plugins/Editor/BlueprintHeaderView/Source/BlueprintHeaderView/Public/BlueprintHeaderView.h
+++ b/Engine/Plugins/Editor/BlueprintHeaderView/Source/BlueprintHeaderView/Public/BlueprintHeaderView.h
@@ -8,6 +8,7 @@
 #include "Styling/SlateTypes.h"
 
 class FExtender;
+class FSlateStyleSet;
 
 class FBlueprintHeaderViewModule : public IModuleInterface
 {
@@ -35,6 +36,9 @@ public:
 	/** TableRow style for the Header View Output */
 	static FTableRowStyle HeaderViewTableRowStyle;
 
+	/** Style set for the header view */
+	static TSharedPtr<FSlateStyleSet> HeaderViewStyleSet;
+
 private: 
 	/** Handle to our delegate so we can remove it at module shutdown */
 	FDelegateHandle ContentBrowserExtenderDelegateHandle;
